Move AnimToTexture from CitySample to Engine plugin

diff --git a/Engine/Plugins/Animation/AnimToTexture/AnimToTexture.uplugin b/Engine/Plugins/Animation/AnimToTexture/AnimToTexture.uplugin
new file mode 100644
--- /dev/null
+++ b/Engine/Plugins/Animation/AnimToTexture/AnimToTexture.uplugin
@@ -0,0 +1,29 @@
+{
+	"FileVersion": 3,
+	"Version": 1,
+	"VersionName": "1.0",
+	"FriendlyName": "AnimToTexture",
+	"Description": "Converts SkeletalMesh Animations into Textures",
+	"Category": "Other",
+	"CreatedBy": "Epic Games",
+	"CreatedByURL": "",
+	"DocsURL": "",
+	"MarketplaceURL": "",
+	"SupportURL": "",
+	"CanContainContent": true,
+	"IsBetaVersion": true,
+	"IsExperimentalVersion": false,
+	"Installed": false,
+	"Modules": [
+		{
+			"Name": "AnimToTexture",
+			"Type": "Runtime",
+			"LoadingPhase": "PreDefault"
+		},
+		{
+			"Name": "AnimToTextureEditor",
+			"Type": "Editor",
+			"LoadingPhase": "PreLoadingScreen"
+		}
+	]
+}
\ No newline at end of file
diff --git a/Engine/Plugins/Animation/AnimToTexture/Config/DefaultAnimToTexture.ini b/Engine/Plugins/Animation/AnimToTexture/Config/DefaultAnimToTexture.ini
new file mode 100644
--- /dev/null
+++ b/Engine/Plugins/Animation/AnimToTexture/Config/DefaultAnimToTexture.ini
@@ -0,0 +1,3 @@
+[CoreRedirects]
++PropertyRedirects==(OldName="AnimToTextureDataAsset.MasterSkeletalMesh",NewName="LeaderSkeletalMesh")
++FunctionRedirects=(OldName="AnimToTextureDataAsset.BP_GetMasterSkeletalMesh", NewName="BP_GetLeaderSkeletalMesh")
\ No newline at end of file
diff --git a/Engine/Plugins/Animation/AnimToTexture/Source/AnimToTexture/AnimToTexture.Build.cs b/Engine/Plugins/Animation/AnimToTexture/Source/AnimToTexture/AnimToTexture.Build.cs
new file mode 100644
--- /dev/null
+++ b/Engine/Plugins/Animation/AnimToTexture/Source/AnimToTexture/AnimToTexture.Build.cs
@@ -0,0 +1,54 @@
+// Copyright Epic Games, Inc. All Rights Reserved.
+
+using UnrealBuildTool;
+
+public class AnimToTexture : ModuleRules
+{
+	public AnimToTexture(ReadOnlyTargetRules Target) : base(Target)
+	{
+		PCHUsage = ModuleRules.PCHUsageMode.UseExplicitOrSharedPCHs;
+		
+		PublicIncludePaths.AddRange(
+			new string[] {
+				// ... add public include paths required here ...
+			}
+			);
+				
+		
+		PrivateIncludePaths.AddRange(
+			new string[] {
+				// ... add other private include paths required here ...
+			}
+			);
+			
+		
+		PublicDependencyModuleNames.AddRange(
+			new string[]
+			{
+				"Core",
+				// ... add other public dependencies that you statically link with here ...
+			}
+			);
+			
+		
+		PrivateDependencyModuleNames.AddRange(
+			new string[]
+			{
+				"CoreUObject",
+				"Engine",
+				"Slate",
+				"SlateCore",
+				// ... add private dependencies that you statically link with here ...	
+				"MeshDescription",
+			}
+			);
+		
+		
+		DynamicallyLoadedModuleNames.AddRange(
+			new string[]
+			{
+				// ... add any modules that your module loads dynamically here ...
+			}
+			);
+	}
+}
diff --git a/Engine/Plugins/Animation/AnimToTexture/Source/AnimToTexture/Private/AnimToTexture.cpp b/Engine/Plugins/Animation/AnimToTexture/Source/AnimToTexture/Private/AnimToTexture.cpp
new file mode 100644
--- /dev/null
+++ b/Engine/Plugins/Animation/AnimToTexture/Source/AnimToTexture/Private/AnimToTexture.cpp
@@ -0,0 +1,22 @@
+// Copyright Epic Games, Inc. All Rights Reserved.
+
+#include "AnimToTexture.h"
+
+#define LOCTEXT_NAMESPACE "FAnimToTextureModule"
+
+void FAnimToTextureModule::StartupModule()
+{
+	// This code will execute after your module is loaded into memory; the exact timing is specified in the .uplugin file per-module
+	
+}
+
+void FAnimToTextureModule::ShutdownModule()
+{
+	// This function may be called during shutdown to clean up your module.  For modules that support dynamic reloading,
+	// we call this function before unloading the module.
+	
+}
+
+#undef LOCTEXT_NAMESPACE
+	
+IMPLEMENT_MODULE(FAnimToTextureModule, AnimToTexture)
\ No newline at end of file
diff --git a/Engine/Plugins/Animation/AnimToTexture/Source/AnimToTexture/Private/AnimToTextureDataAsset.cpp b/Engine/Plugins/Animation/AnimToTexture/Source/AnimToTexture/Private/AnimToTextureDataAsset.cpp
new file mode 100644
--- /dev/null
+++ b/Engine/Plugins/Animation/AnimToTexture/Source/AnimToTexture/Private/AnimToTextureDataAsset.cpp
@@ -0,0 +1,55 @@
+// Copyright Epic Games, Inc. All Rights Reserved.
+
+#include "AnimToTextureDataAsset.h"
+
+int32 UAnimToTextureDataAsset::GetIndexFromAnimSequence(const UAnimSequence* Sequence)
+{
+	int32 OutIndex = 0;
+
+	int32 NumSequences = AnimSequences.Num();
+	
+	// We can store a sequence to index map for a faster search
+	for (int32 CurrentIndex = 0; CurrentIndex < NumSequences; ++CurrentIndex)
+	{
+		const FAnimSequenceInfo& SequenceInfo = AnimSequences[CurrentIndex];
+		if (SequenceInfo.AnimSequence == Sequence)
+		{
+			OutIndex = CurrentIndex;
+			break;
+		}
+	}
+
+	return OutIndex;
+}
+
+namespace AnimToTextureParamNames
+{
+	static const FName BoundingBoxMin = TEXT("MinBBox");
+	static const FName BoundingBoxScale = TEXT("SizeBBox");
+	static const FName RowsPerFrame = TEXT("RowsPerFrame");
+	static const FName BoneWeightRowsPerFrame = TEXT("BoneWeightsRowsPerFrame");
+	static const FName NumFrames = TEXT("NumFrames (S)");
+	//static const FName UVChannel = TEXT("UVChannel");
+	//static const FName AnimateSwitch = TEXT("Animate (B)");
+	static const FName VertexPositionTexture = TEXT("PositionTexture");
+	static const FName VertexNormalTexture = TEXT("NormalTexture");
+	static const FName BonePositionTexture = TEXT("BonePositionTexture");
+	static const FName BoneRotationTexture = TEXT("BoneRotationTexture");
+	static const FName BoneWeightsTexture = TEXT("BoneWeightsTexture");
+}
+
+FAnimToTextureMaterialParamNames::FAnimToTextureMaterialParamNames()
+{
+	BoundingBoxMin = AnimToTextureParamNames::BoundingBoxMin;
+	BoundingBoxScale = AnimToTextureParamNames::BoundingBoxScale;
+	RowsPerFrame = AnimToTextureParamNames::RowsPerFrame;
+	BoneWeightRowsPerFrame = AnimToTextureParamNames::BoneWeightRowsPerFrame;
+	NumFrames = AnimToTextureParamNames::NumFrames;
+	//UVChannel = AnimToTextureParamNames::UVChannel;
+	//AnimateSwitch = AnimToTextureParamNames::AnimateSwitch;
+	VertexPositionTexture = AnimToTextureParamNames::VertexPositionTexture;
+	VertexNormalTexture = AnimToTextureParamNames::VertexNormalTexture;
+	BonePositionTexture = AnimToTextureParamNames::BonePositionTexture;
+	BoneRotationTexture = AnimToTextureParamNames::BoneRotationTexture;
+	BoneWeightsTexture = AnimToTextureParamNames::BoneWeightsTexture;
+}
diff --git a/Engine/Plugins/Animation/AnimToTexture/Source/AnimToTexture/Private/AnimToTextureInstancePlaybackHelpers.cpp b/Engine/Plugins/Animation/AnimToTexture/Source/AnimToTexture/Private/AnimToTextureInstancePlaybackHelpers.cpp
new file mode 100644
--- /dev/null
+++ b/Engine/Plugins/Animation/AnimToTexture/Source/AnimToTexture/Private/AnimToTextureInstancePlaybackHelpers.cpp
@@ -0,0 +1,83 @@
+// Copyright Epic Games, Inc. All Rights Reserved.
+
+#include "AnimToTextureInstancePlaybackHelpers.h"
+#include "AnimToTextureDataAsset.h"
+
+void UAnimToTextureInstancePlaybackLibrary::SetupInstancedMeshComponent(UInstancedStaticMeshComponent* InstancedMeshComponent, FAnimToTextureInstanceData& InstanceData, int32 NumInstances)
+{
+	if (InstancedMeshComponent)
+	{
+		InstancedMeshComponent->NumCustomDataFloats = InstanceData.PlaybackData.GetTypeSize() / sizeof(float);
+		InstancedMeshComponent->PerInstanceSMData.Reset();
+		InstancedMeshComponent->PerInstanceSMData.AddDefaulted(NumInstances);
+		InstancedMeshComponent->PreAllocateInstancesMemory(NumInstances);
+		InstancedMeshComponent->PerInstanceSMCustomData.SetNumZeroed(NumInstances * InstancedMeshComponent->NumCustomDataFloats);
+		AllocateInstanceData(InstanceData, NumInstances);
+	}
+}
+
+void UAnimToTextureInstancePlaybackLibrary::BatchUpdateInstancedMeshComponent(UInstancedStaticMeshComponent* InstancedMeshComponent, FAnimToTextureInstanceData& InstanceData)
+{
+	SIZE_T CustomDataSizeToCopy = FMath::Min(InstanceData.PlaybackData.Num() * InstanceData.PlaybackData.GetTypeSize(), InstancedMeshComponent->PerInstanceSMCustomData.Num() * InstancedMeshComponent->PerInstanceSMCustomData.GetTypeSize());
+	FMemory::Memcpy(InstancedMeshComponent->PerInstanceSMCustomData.GetData(), InstanceData.PlaybackData.GetData(), CustomDataSizeToCopy);
+
+	int32 TransformsToCopy = FMath::Min(InstancedMeshComponent->GetNumRenderInstances(), InstanceData.StaticMeshInstanceData.Num());
+	InstancedMeshComponent->BatchUpdateInstancesData(0, TransformsToCopy, InstanceData.StaticMeshInstanceData.GetData(), true, false);
+}
+
+void UAnimToTextureInstancePlaybackLibrary::AllocateInstanceData(FAnimToTextureInstanceData& InstanceData, int32 Count)
+{
+	InstanceData.StaticMeshInstanceData.AddDefaulted(Count);
+	InstanceData.PlaybackData.AddDefaulted(Count);
+}
+
+bool UAnimToTextureInstancePlaybackLibrary::UpdateInstanceData(FAnimToTextureInstanceData& InstanceData, int32 InstanceIndex, const FAnimToTextureInstancePlaybackData& PlaybackData, const FTransform& Transform)
+{
+	if (InstanceData.PlaybackData.IsValidIndex(InstanceIndex) && InstanceData.StaticMeshInstanceData.IsValidIndex(InstanceIndex))
+	{
+		InstanceData.PlaybackData[InstanceIndex] = PlaybackData;
+		InstanceData.StaticMeshInstanceData[InstanceIndex].Transform = Transform.ToMatrixWithScale();
+		return true;
+	}
+
+	return false;
+}
+
+bool UAnimToTextureInstancePlaybackLibrary::GetInstancePlaybackData(const FAnimToTextureInstanceData& InstanceData, int32 InstanceIndex, FAnimToTextureInstancePlaybackData& InstancePlaybackData)
+{
+	if (InstanceData.PlaybackData.IsValidIndex(InstanceIndex))
+	{
+		InstancePlaybackData = InstanceData.PlaybackData[InstanceIndex];
+		return true;
+	}
+	
+	return false;
+}
+
+bool UAnimToTextureInstancePlaybackLibrary::GetInstanceTransform(const FAnimToTextureInstanceData& InstanceData, int32 InstanceIndex, FTransform& InstanceTransform)
+{
+	if (InstanceData.StaticMeshInstanceData.IsValidIndex(InstanceIndex))
+	{
+		InstanceTransform = FTransform(InstanceData.StaticMeshInstanceData[InstanceIndex].Transform);
+		return true;
+	}
+
+	return false;
+}
+
+bool UAnimToTextureInstancePlaybackLibrary::AnimStateFromDataAsset(const UAnimToTextureDataAsset* DataAsset, int32 StateIndex, FAnimToTextureAnimState& AnimState)
+{
+	if (DataAsset && DataAsset->Animations.IsValidIndex(StateIndex))
+	{
+		const FAnimInfo& AnimInfo = DataAsset->Animations[StateIndex];
+		AnimState.StartFrame = AnimInfo.AnimStart;
+		AnimState.NumFrames = AnimInfo.NumFrames;
+		AnimState.bLooping = AnimInfo.bLooping;
+
+		return true;
+	}
+
+	AnimState = FAnimToTextureAnimState();
+	return false;
+}
+
diff --git a/Engine/Plugins/Animation/AnimToTexture/Source/AnimToTexture/Private/EvaluateSequenceAnimInstance.cpp b/Engine/Plugins/Animation/AnimToTexture/Source/AnimToTexture/Private/EvaluateSequenceAnimInstance.cpp
new file mode 100644
--- /dev/null
+++ b/Engine/Plugins/Animation/AnimToTexture/Source/AnimToTexture/Private/EvaluateSequenceAnimInstance.cpp
@@ -0,0 +1,10 @@
+// Copyright Epic Games, Inc. All Rights Reserved.
+
+#include "EvaluateSequenceAnimInstance.h"
+
+
+UEvaluateSequenceAnimInstance::UEvaluateSequenceAnimInstance(const FObjectInitializer& ObjectInitializer)
+	: Super(ObjectInitializer)
+{
+
+}
\ No newline at end of file
diff --git a/Engine/Plugins/Animation/AnimToTexture/Source/AnimToTexture/Private/LightweightMontageInstance.cpp b/Engine/Plugins/Animation/AnimToTexture/Source/AnimToTexture/Private/LightweightMontageInstance.cpp
new file mode 100644
--- /dev/null
+++ b/Engine/Plugins/Animation/AnimToTexture/Source/AnimToTexture/Private/LightweightMontageInstance.cpp
@@ -0,0 +1,259 @@
+// Copyright Epic Games, Inc. All Rights Reserved.
+
+#include "LightweightMontageInstance.h"
+#include "Animation/AnimMontage.h"
+#include "Animation/AnimSequence.h"
+
+EMontageSubStepResult UE::VertexAnimation::FLightweightMontageSubstepper::Advance(float& InOut_P_Original, UAnimMontage* Montage)
+{
+	// Copied from MontageSubstepper::Advance
+	// @todo : Clean up + optimize when we do profiling
+
+	DeltaMove = 0.f;
+
+	if (Montage == nullptr)
+	{
+		return EMontageSubStepResult::InvalidMontage;
+	}
+
+	bReachedEndOfSection = false;
+
+	float PositionInSection;
+	CurrentSectionIndex = Montage->GetAnimCompositeSectionIndexFromPos(InOut_P_Original, PositionInSection);
+	if (!Montage->IsValidSectionIndex(CurrentSectionIndex))
+	{
+		return EMontageSubStepResult::InvalidSection;
+	}
+
+	const FCompositeSection& CurrentSection = Montage->GetAnimCompositeSection(CurrentSectionIndex);
+	CurrentSectionStartTime = CurrentSection.GetTime();
+
+	// Find end of current section. We only update one section at a time.
+	CurrentSectionLength = Montage->GetSectionLength(CurrentSectionIndex);
+
+	if (FMath::IsNearlyZero(TimeRemaining))
+	{
+		return EMontageSubStepResult::NotMoved;
+	}
+
+	const float PlayRate = Montage->RateScale;
+
+	if (FMath::IsNearlyZero(PlayRate))
+	{
+		return EMontageSubStepResult::NotMoved;
+	}
+
+	DeltaMove = TimeRemaining * PlayRate;
+
+	// Finally clamp DeltaMove by section markers.
+	{
+		const float OldDeltaMove = DeltaMove;
+
+		// Clamp DeltaMove based on move allowed within current section
+		// We stop at each section marker to evaluate whether we should jump to another section marker or not.
+		// Test is inclusive, so we know if we've reached marker or not.
+		const float MaxSectionMove = CurrentSectionLength - PositionInSection;
+		if (DeltaMove >= MaxSectionMove)
+		{
+			DeltaMove = MaxSectionMove;
+			bReachedEndOfSection = true;
+		}
+	}
+
+	// DeltaMove is now final, see if it has any effect on our position.
+	if (FMath::Abs(DeltaMove) > 0.f)
+	{
+		// Note that we don't worry about looping and wrapping around here.
+		// We step per section to simplify code to extract notifies/root motion/etc.
+		InOut_P_Original += DeltaMove;
+
+		// Decrease RemainingTime with actual time elapsed 
+		// So we can take more substeps as needed.
+		const float TimeStep = DeltaMove / PlayRate;
+		ensure(TimeStep >= 0.f);
+		TimeRemaining = FMath::Max(TimeRemaining - TimeStep, 0.f);
+
+		return EMontageSubStepResult::Moved;
+	}
+	else
+	{
+		return EMontageSubStepResult::NotMoved;
+	}
+}
+
+void UE::VertexAnimation::FLightweightMontageInstance::Initialize(UAnimMontage* InMontage, float StartTime)
+{
+	Terminate();
+
+	MontageWeak = InMontage;
+	Montage = InMontage;
+	if (Montage)
+	{
+		CachedLength = InMontage->GetPlayLength();
+		bInitialized = true;
+		Position = StartTime;
+		RefreshNextPrevSections(Montage);
+		MontageSubStepper.Initialize(this);
+		MontageSubStepper.Advance(Position, Montage);
+		SetSequence(FindSequenceAtCurrentTime(SectionStartTimeTransient));
+	}
+
+	Montage = nullptr;
+}
+
+void UE::VertexAnimation::FLightweightMontageInstance::RefreshNextPrevSections(UAnimMontage* InMontage)
+{
+	// initialize next section
+	if (InMontage->CompositeSections.Num() > 0)
+	{
+		NextSections.Empty(InMontage->CompositeSections.Num());
+		NextSections.AddUninitialized(InMontage->CompositeSections.Num());
+
+		for (int32 I = 0; I < InMontage->CompositeSections.Num(); ++I)
+		{
+			FCompositeSection& Section = InMontage->CompositeSections[I];
+			int32 NextSectionIdx = InMontage->GetSectionIndex(Section.NextSectionName);
+			NextSections[I] = NextSectionIdx;
+		}
+	}
+}
+
+const UAnimSequence* UE::VertexAnimation::FLightweightMontageInstance::FindSequenceAtCurrentTime(float& SequenceStartTime)
+{
+	SequenceStartTime = 0.0f;
+	const FAnimTrack& AnimTrack = Montage->SlotAnimTracks[0].AnimTrack;
+	if (const FAnimSegment* Segment = AnimTrack.GetSegmentAtTime(Position))
+	{
+		SequenceStartTime = Segment->StartPos - Segment->AnimStartTime;
+		return Cast<UAnimSequence>(Segment->GetAnimReference());
+	}
+
+	return nullptr;
+}
+
+UAnimMontage* UE::VertexAnimation::FLightweightMontageInstance::GetMontage() const
+{
+	return bInitialized ? MontageWeak.Get() : nullptr;
+}
+
+const UAnimSequence* UE::VertexAnimation::FLightweightMontageInstance::GetSequence() const
+{
+	return bInitialized ? SequenceWeak.Get() : nullptr;
+}
+
+void UE::VertexAnimation::FLightweightMontageInstance::SetSequence(const UAnimSequence* InSequence)
+{
+	bSequenceChangedThisFrame = false;
+	const UAnimSequence* PrevSequence = SequenceWeak.Get();
+	if (PrevSequence != InSequence)
+	{
+		bSequenceChangedThisFrame = true;
+		SequenceWeak = InSequence;
+	}
+}
+
+float UE::VertexAnimation::FLightweightMontageInstance::GetPositionInSection() const
+{
+	return Position - SectionStartTimeTransient;
+}
+
+void UE::VertexAnimation::FLightweightMontageInstance::Terminate()
+{
+	*this = FLightweightMontageInstance();
+}
+
+bool UE::VertexAnimation::FLightweightMontageInstance::Advance_Internal(float DeltaTime, FRootMotionMovementParams& OutRootMotionParams, const FLightWeightMontageExtractionSettings& ExtractionSettings)
+{
+	// Copied from FAnimMontageInstance::Advance
+	// @todo : Clean up + optimize when we do profiling
+
+	const bool bExtractRootMotion = Montage->HasRootMotion() && ExtractionSettings.bExtractRootMotion;
+	
+	/**
+		Limit number of iterations for performance.
+		This can get out of control if PlayRate is set really high, or there is a hitch, and Montage is looping for example.
+	*/
+	const int32 MaxIterations = 10;
+	int32 NumIterations = 0;
+
+	/**
+		If we're hitting our max number of iterations for whatever reason,
+		make sure we're not accumulating too much time, and go out of range.
+	*/
+	if (MontageSubStepper.GetRemainingTime() < 10.f)
+	{
+		MontageSubStepper.AddEvaluationTime(DeltaTime);
+	}
+
+	while (MontageSubStepper.HasTimeRemaining() && (++NumIterations < MaxIterations))
+	{
+		const float PreviousSubStepPosition = Position;
+		EMontageSubStepResult SubStepResult = MontageSubStepper.Advance(Position, Montage);
+
+		if (SubStepResult != EMontageSubStepResult::Moved)
+		{
+			// stop and leave this loop
+			break;
+		}
+
+		// Extract Root Motion for this time slice, and accumulate it.
+		if (bExtractRootMotion)
+		{
+			OutRootMotionParams.Accumulate(Montage->ExtractRootMotionFromTrackRange(PreviousSubStepPosition, Position));
+		}
+
+		// if we reached end of section, and we were not processing a branching point, and no events has messed with out current position..
+		// .. Move to next section.
+		// (this also handles looping, the same as jumping to a different section).
+		if (MontageSubStepper.HasReachedEndOfSection())
+		{
+			const int32 CurrentSectionIndex = MontageSubStepper.GetCurrentSectionIndex();
+
+			// Get recent NextSectionIndex in case it's been changed by previous events.
+			const int32 RecentNextSectionIndex = NextSections[CurrentSectionIndex];
+			if (RecentNextSectionIndex != INDEX_NONE)
+			{
+				float LatestNextSectionStartTime;
+				float LatestNextSectionEndTime;
+				Montage->GetSectionStartAndEndTime(RecentNextSectionIndex, LatestNextSectionStartTime, LatestNextSectionEndTime);
+
+				// Jump to next section's appropriate starting point (start or end).
+				Position = LatestNextSectionStartTime;
+			}
+			else
+			{
+				// Reached end of last section. Exit.
+				return false;
+				break;
+			}
+		}
+	}
+
+	return true;
+}
+
+bool UE::VertexAnimation::FLightweightMontageInstance::IsValid() const
+{
+	return bInitialized && MontageWeak.IsValid() && SequenceWeak.IsValid();
+}
+
+void UE::VertexAnimation::FLightweightMontageInstance::Advance(float DeltaTime, float InGlobalTime, FRootMotionMovementParams& OutRootMotionParams, const FLightWeightMontageExtractionSettings& ExtractionSettings)
+{
+	Montage = GetMontage();
+	if (Montage == nullptr)
+	{
+		return;
+	}
+
+	const bool bIsPlaying = Advance_Internal(DeltaTime, OutRootMotionParams, ExtractionSettings);
+	const UAnimSequence* CurrentSequence = FindSequenceAtCurrentTime(SectionStartTimeTransient);
+
+	SetSequence(CurrentSequence);
+
+	Montage = nullptr;
+
+	if (!bIsPlaying)
+	{
+		Terminate();
+	}
+}
\ No newline at end of file
diff --git a/Engine/Plugins/Animation/AnimToTexture/Source/AnimToTexture/Public/AnimToTexture.h b/Engine/Plugins/Animation/AnimToTexture/Source/AnimToTexture/Public/AnimToTexture.h
new file mode 100644
--- /dev/null
+++ b/Engine/Plugins/Animation/AnimToTexture/Source/AnimToTexture/Public/AnimToTexture.h
@@ -0,0 +1,14 @@
+// Copyright Epic Games, Inc. All Rights Reserved.
+
+#pragma once
+
+#include "Modules/ModuleManager.h"
+
+class FAnimToTextureModule : public IModuleInterface
+{
+public:
+
+	/** IModuleInterface implementation */
+	virtual void StartupModule() override;
+	virtual void ShutdownModule() override;
+};
diff --git a/Engine/Plugins/Animation/AnimToTexture/Source/AnimToTexture/Public/AnimToTextureDataAsset.h b/Engine/Plugins/Animation/AnimToTexture/Source/AnimToTexture/Public/AnimToTextureDataAsset.h
new file mode 100644
--- /dev/null
+++ b/Engine/Plugins/Animation/AnimToTexture/Source/AnimToTexture/Public/AnimToTextureDataAsset.h
@@ -0,0 +1,358 @@
+// Copyright Epic Games, Inc. All Rights Reserved.
+
+#pragma once
+
+#include "Engine/DataAsset.h"
+#include "Animation/AnimInstance.h"
+#include "Engine/StaticMesh.h"
+#include "AnimToTextureDataAsset.generated.h"
+
+class USkeletalMesh;
+class UStaticMesh;
+class UTexture2D;
+
+USTRUCT(Blueprintable)
+struct ANIMTOTEXTURE_API FAnimToTextureMaterialParamNames
+{
+	GENERATED_BODY()
+
+	//
+	// Static Switch Parameters
+	//
+
+	// UPROPERTY(BlueprintReadWrite, EditAnywhere)
+	// FName AnimateSwitch;
+
+	//
+	// Scalar Parameters
+	//
+
+	//UPROPERTY(EditAnywhere)
+	//FName UVChannel;
+
+	UPROPERTY(BlueprintReadWrite, Category = Default, EditAnywhere)
+	FName RowsPerFrame;
+
+	UPROPERTY(BlueprintReadWrite, Category = Default, EditAnywhere)
+	FName BoneWeightRowsPerFrame;
+
+	UPROPERTY(BlueprintReadWrite, Category = Default, EditAnywhere)
+	FName NumFrames;
+
+	//
+	// Vector Parameters
+	//
+
+	UPROPERTY(BlueprintReadWrite, Category = Default, EditAnywhere)
+	FName BoundingBoxMin;
+
+	UPROPERTY(BlueprintReadWrite, Category = Default, EditAnywhere)
+	FName BoundingBoxScale;
+
+	//
+	// Texture Parameters
+	//
+
+	UPROPERTY(BlueprintReadWrite, Category = Default, EditAnywhere)
+	FName VertexPositionTexture;
+
+	UPROPERTY(BlueprintReadWrite, Category = Default, EditAnywhere)
+	FName VertexNormalTexture;
+
+	UPROPERTY(BlueprintReadWrite, Category = Default, EditAnywhere)
+	FName BonePositionTexture;
+
+	UPROPERTY(BlueprintReadWrite, Category = Default, EditAnywhere)
+	FName BoneRotationTexture;
+
+	UPROPERTY(BlueprintReadWrite, Category = Default, EditAnywhere)
+	FName BoneWeightsTexture;
+
+	// Initialize Names
+	FAnimToTextureMaterialParamNames();
+};
+
+UENUM(Blueprintable)
+enum class EAnimToTextureMode : uint8
+{
+	/* Position and Normal Per-Vertex */
+	Vertex,
+	/* Linear Blending Skinnin */
+	Bone,
+};
+
+//UENUM(Blueprintable)
+//enum class EAnimToTextureNumInfluences : uint8
+//{
+//	One, Two, Four,
+//};
+
+UENUM(Blueprintable)
+enum class EAnimToTextureBonePrecision : uint8
+{
+	/* Bone positions and rotations stored in 8 bits */
+	EightBits,
+	/* Bone positions and rotations stored in 16 bits */
+	SixteenBits,
+};
+
+
+//USTRUCT(Blueprintable)
+//struct FSkeletalMeshInfo
+//{
+//	GENERATED_BODY()
+//
+//	UPROPERTY(EditAnywhere)
+//	USkeletalMesh* SkeletalMesh = nullptr;
+//
+//	UPROPERTY(EditAnywhere)
+//	int32 LODIndex = 0;
+//};
+
+//USTRUCT(Blueprintable)
+//struct FStaticMeshInfo
+//{
+//	GENERATED_BODY()
+//
+//	UPROPERTY(EditAnywhere)
+//	UStaticMesh* StaticMesh = nullptr;
+//
+//	UPROPERTY(EditAnywhere)
+//	int32 LODIndex = 0;
+//
+//	UPROPERTY(EditAnywhere)
+//	int32 UVChannel = 1;
+//};
+
+USTRUCT(Blueprintable)
+struct FAnimSequenceInfo
+{
+	GENERATED_BODY()
+
+	UPROPERTY(EditAnywhere, Category = Default, BlueprintReadWrite)
+	bool bEnabled = true;
+
+	UPROPERTY(EditAnywhere, Category = Default, BlueprintReadWrite)
+	UAnimSequence* AnimSequence = nullptr;
+
+	UPROPERTY(EditAnywhere, Category = Default, BlueprintReadWrite)
+	bool bLooping = true;
+
+	UPROPERTY(EditAnywhere, Category = Default, BlueprintReadWrite)
+	bool bUseCustomRange = false;
+
+	UPROPERTY(EditAnywhere, Category = Default, BlueprintReadWrite, meta = (EditCondition = "bUseCustomRange"))
+	int32 StartFrame = 0;
+
+	UPROPERTY(EditAnywhere, Category = Default, BlueprintReadWrite, meta = (EditCondition = "bUseCustomRange"))
+	int32 EndFrame = 1;
+
+};
+
+USTRUCT(Blueprintable)
+struct FAnimInfo
+{
+	GENERATED_BODY()
+
+	UPROPERTY(VisibleAnywhere, Category = Default)
+	int32 NumFrames = 0;
+
+	UPROPERTY(VisibleAnywhere, Category = Default)
+	int32 AnimStart = 0;
+
+	UPROPERTY(EditAnywhere, Category = Default)
+	bool bLooping = true;
+};
+
+UCLASS(Blueprintable, BlueprintType)
+class ANIMTOTEXTURE_API UAnimToTextureDataAsset : public UPrimaryDataAsset
+{
+public:
+	GENERATED_BODY()
+
+	// ------------------------------------------------------
+	// Skeletal Mesh
+
+	UPROPERTY(EditAnywhere, BlueprintReadWrite, Category = "SkeletalMesh", meta = (AssetBundles = "Client"))
+	TSoftObjectPtr<USkeletalMesh> SkeletalMesh;
+
+	UPROPERTY(EditAnywhere, BlueprintReadWrite, Category = "SkeletalMesh", Meta = (DisplayName = "LODIndex"))
+	int32 SkeletalLODIndex = 0;
+
+	// ------------------------------------------------------
+	// Static Mesh
+
+	UPROPERTY(EditAnywhere, BlueprintReadWrite, Category = "StaticMesh", meta = (AssetBundles = "Client"))
+	TSoftObjectPtr<UStaticMesh> StaticMesh;
+
+	UPROPERTY(EditAnywhere, BlueprintReadWrite, Category = "StaticMesh", Meta = (DisplayName = "LODIndex"))
+	int32 StaticLODIndex = 0;
+
+	UPROPERTY(EditAnywhere, BlueprintReadWrite, Category = "StaticMesh")
+	int32 UVChannel = 1;
+
+	// ------------------------------------------------------
+	// Texture
+
+	UPROPERTY(EditAnywhere, BlueprintReadWrite, Category = "Texture")
+	int32 MaxHeight = 4096;
+
+	UPROPERTY(EditAnywhere, BlueprintReadWrite, Category = "Texture")
+	int32 MaxWidth = 4096;
+
+	UPROPERTY(EditAnywhere, BlueprintReadWrite, Category = "Texture")
+	bool bEnforcePowerOfTwo = false;
+
+	UPROPERTY(EditAnywhere, BlueprintReadWrite, Category = "Texture")
+	EAnimToTextureMode Mode;
+
+	UPROPERTY(EditAnywhere, BlueprintReadWrite, Category = "Texture|Vertex", meta = (EditCondition = "Mode == EAnimToTextureMode::Vertex"))
+	TSoftObjectPtr<UTexture2D> VertexPositionTexture;
+
+	UPROPERTY(EditAnywhere, BlueprintReadWrite, Category = "Texture|Vertex", meta = (EditCondition = "Mode == EAnimToTextureMode::Vertex"))
+	TSoftObjectPtr<UTexture2D> VertexNormalTexture;
+
+	// UPROPERTY(EditAnywhere, BlueprintReadWrite, Category = "Texture|Bone", meta = (EditCondition = "Mode == EAnimToTextureMode::Bone"))
+	// EAnimToTextureNumInfluences NumInfluences;
+	
+	UPROPERTY(EditAnywhere, BlueprintReadWrite, Category = "Texture|Bone", meta = (EditCondition = "Mode == EAnimToTextureMode::Bone"))
+	TSoftObjectPtr<UTexture2D> BonePositionTexture;
+
+	UPROPERTY(EditAnywhere, BlueprintReadWrite, Category = "Texture|Bone", meta = (EditCondition = "Mode == EAnimToTextureMode::Bone"))
+	TSoftObjectPtr<UTexture2D> BoneRotationTexture;
+
+	UPROPERTY(EditAnywhere, BlueprintReadWrite, Category = "Texture|Bone", meta = (EditCondition = "Mode == EAnimToTextureMode::Bone"))
+	TSoftObjectPtr<UTexture2D> BoneWeightTexture;
+
+	UPROPERTY(EditAnywhere, BlueprintReadWrite, Category = "Texture|Bone", meta = (EditCondition = "Mode == EAnimToTextureMode::Bone"))
+	EAnimToTextureBonePrecision PositionAndRotationPrecision = EAnimToTextureBonePrecision::EightBits;
+	
+	// ------------------------------------------------------
+	// Animation
+
+	 /** This Mesh will be used as LeaderPose. Animations must use same Skeleton than this SkeletalMesh */
+	UPROPERTY(EditAnywhere, BlueprintReadWrite, Category = "Animation")
+	TSoftObjectPtr<USkeletalMesh> LeaderSkeletalMesh;
+
+	/** Bone used for Rigid Binding. The bone needs to be part of the RawBones. 
+	*   Sockets and VirtualBones are not supported.
+	*/
+	UPROPERTY(EditAnywhere, BlueprintReadWrite, Category = "Animation")
+	FName AttachToSocket;
+
+	UPROPERTY(EditAnywhere, BlueprintReadWrite, Category = "Animation")
+	TSubclassOf<UAnimInstance> AnimInstanceClass;
+
+	UPROPERTY(EditAnywhere, BlueprintReadWrite, Category = "Animation")
+	float SampleRate = 30.0f;
+
+	UPROPERTY(EditAnywhere, BlueprintReadWrite, Category = "Animation")
+	TArray<FAnimSequenceInfo> AnimSequences;
+
+	// ------------------------------------------------------
+	// Info
+
+	/* Total Number of Frames in all animations */
+	UPROPERTY(VisibleAnywhere, BlueprintReadOnly, Category = "Info")
+	int32 NumFrames = 0;
+
+	UPROPERTY(VisibleAnywhere, BlueprintReadOnly, Category = "Info|Vertex", Meta = (DisplayName = "RowsPerFrame"))
+	int32 VertexRowsPerFrame = 1;
+
+	UPROPERTY(VisibleAnywhere, BlueprintReadOnly, Category = "Info|Vertex", Meta = (DisplayName = "MinBBox"))
+	FVector VertexMinBBox;
+
+	UPROPERTY(VisibleAnywhere, BlueprintReadOnly, Category = "Info|Vertex", Meta = (DisplayName = "SizeBBox"))
+	FVector VertexSizeBBox;
+
+	UPROPERTY(VisibleAnywhere, BlueprintReadOnly, Category = "Info|Bone")
+	int32 BoneWeightRowsPerFrame = 1;
+
+	UPROPERTY(VisibleAnywhere, BlueprintReadOnly, Category = "Info|Bone")
+	int32 BoneRowsPerFrame = 1;
+
+	UPROPERTY(VisibleAnywhere, BlueprintReadOnly, Category = "Info|Bone", Meta = (DisplayName = "MinBBox"))
+	FVector BoneMinBBox;
+
+	UPROPERTY(VisibleAnywhere, BlueprintReadOnly, Category = "Info|Bone", Meta = (DisplayName = "SizeBBox"))
+	FVector BoneSizeBBox;
+
+	UPROPERTY(VisibleAnywhere, BlueprintReadOnly, Category = "Info")
+	TArray<FAnimInfo> Animations;
+
+	UFUNCTION(BlueprintCallable, Category = Default)
+	int32 GetIndexFromAnimSequence(const UAnimSequence* Sequence);
+
+	UFUNCTION()
+	void Reset();
+
+	// If we weren't in a plugin, we could unify this in a base class
+	template<typename AssetType>
+	static AssetType* GetAsset(const TSoftObjectPtr<AssetType>& AssetPointer)
+	{
+		AssetType* ReturnVal = nullptr;
+		if (AssetPointer.ToSoftObjectPath().IsValid())
+		{
+			ReturnVal = AssetPointer.Get();
+			if (!ReturnVal)
+			{
+				AssetType* LoadedAsset = Cast<AssetType>(AssetPointer.ToSoftObjectPath().TryLoad());
+				if (ensureMsgf(LoadedAsset, TEXT("Failed to load asset pointer %s"), *AssetPointer.ToString()))
+				{
+					ReturnVal = LoadedAsset;
+				}
+			}
+		}
+		return ReturnVal;
+	}
+
+#define AnimToTextureDataAsset_ASSET_ACCESSOR(ClassName, PropertyName) \
+	FORCEINLINE ClassName* Get##PropertyName() const { return GetAsset(PropertyName); }
+
+	AnimToTextureDataAsset_ASSET_ACCESSOR(UStaticMesh, StaticMesh);
+	AnimToTextureDataAsset_ASSET_ACCESSOR(USkeletalMesh, SkeletalMesh);
+	AnimToTextureDataAsset_ASSET_ACCESSOR(USkeletalMesh, LeaderSkeletalMesh);
+	AnimToTextureDataAsset_ASSET_ACCESSOR(UTexture2D, VertexPositionTexture);
+	AnimToTextureDataAsset_ASSET_ACCESSOR(UTexture2D, VertexNormalTexture);
+	AnimToTextureDataAsset_ASSET_ACCESSOR(UTexture2D, BonePositionTexture);
+	AnimToTextureDataAsset_ASSET_ACCESSOR(UTexture2D, BoneRotationTexture);
+	AnimToTextureDataAsset_ASSET_ACCESSOR(UTexture2D, BoneWeightTexture);
+
+	UFUNCTION(BlueprintPure, Category = Default, meta = (DisplayName = "Get Static Mesh"))
+	UStaticMesh* BP_GetStaticMesh() { return GetStaticMesh(); }
+
+	UFUNCTION(BlueprintPure, Category = Default, meta = (DisplayName = "Get Skeletal Mesh"))
+	USkeletalMesh* BP_GetSkeletalMesh() { return GetSkeletalMesh(); }
+
+	UFUNCTION(BlueprintPure, Category = Default, meta = (DisplayName = "Get Leader Skeletal Mesh"))
+	USkeletalMesh* BP_GetLeaderSkeletalMesh() { return GetLeaderSkeletalMesh(); }
+
+	UFUNCTION(BlueprintPure, Category = Default, meta = (DisplayName = "Get Bone Position Texture"))
+	UTexture2D* BP_GetBonePositionTexture() { return GetBonePositionTexture(); }
+
+	UFUNCTION(BlueprintPure, Category = Default, meta = (DisplayName = "Get Bone Rotation Texture"))
+	UTexture2D* BP_GetBoneRotationTexture() { return GetBoneRotationTexture(); }
+
+	UFUNCTION(BlueprintPure, Category = Default, meta = (DisplayName = "Get Bone Weight Texture"))
+	UTexture2D* BP_GetBoneWeightTexture() { return GetBoneWeightTexture(); }
+};
+
+FORCEINLINE void UAnimToTextureDataAsset::Reset()
+{
+	// Common Info.
+	//this->NumVertices = 0;
+	this->NumFrames = 0;
+	this->Animations.Reset();
+
+	// Vertex Info
+	this->VertexRowsPerFrame = 1;
+	this->VertexMinBBox = FVector::ZeroVector; // { TNumericLimits<float>::Max(), TNumericLimits<float>::Max(), TNumericLimits<float>::Max() };
+	this->VertexSizeBBox = FVector::ZeroVector;
+
+	// Bone Info
+	//this->NumBones = 0;
+	this->BoneRowsPerFrame = 1;
+	this->BoneWeightRowsPerFrame = 1;
+	this->BoneMinBBox = FVector::ZeroVector;  // { TNumericLimits<float>::Max(), TNumericLimits<float>::Max(), TNumericLimits<float>::Max() };
+	this->BoneSizeBBox = FVector::ZeroVector;
+};
\ No newline at end of file
diff --git a/Engine/Plugins/Animation/AnimToTexture/Source/AnimToTexture/Public/AnimToTextureInstancePlaybackHelpers.h b/Engine/Plugins/Animation/AnimToTexture/Source/AnimToTexture/Public/AnimToTextureInstancePlaybackHelpers.h
new file mode 100644
--- /dev/null
+++ b/Engine/Plugins/Animation/AnimToTexture/Source/AnimToTexture/Public/AnimToTextureInstancePlaybackHelpers.h
@@ -0,0 +1,98 @@
+// Copyright Epic Games, Inc. All Rights Reserved.
+
+#pragma once
+
+#include "Kismet/BlueprintFunctionLibrary.h"
+#include "Components/InstancedStaticMeshComponent.h"
+
+#include "AnimToTextureInstancePlaybackHelpers.generated.h"
+
+// Use floats to match custom floats of instanced static mesh
+// We could pack a float w/ more parameters if desired
+USTRUCT(BlueprintType)
+struct FAnimToTextureAnimState
+{
+	GENERATED_USTRUCT_BODY()
+
+	UPROPERTY(EditAnywhere, BlueprintReadWrite, Category = "AnimToTexture")
+	float StartFrame = 0.0f;
+
+	UPROPERTY(EditAnywhere, BlueprintReadWrite, Category = "AnimToTexture")
+	float NumFrames = 60.0f;
+
+	UPROPERTY(EditAnywhere, BlueprintReadWrite, Category = "AnimToTexture")
+	float PlayRate = 1.0f;
+
+	UPROPERTY(EditAnywhere, BlueprintReadWrite, Category = "AnimToTexture")
+	float bLooping = 1.0f; 
+
+	UPROPERTY(EditAnywhere, BlueprintReadWrite, Category = "AnimToTexture")
+	float GlobalStartTime = 0.0f; 
+};
+
+USTRUCT(BlueprintType)
+struct FAnimToTextureInstancePlaybackData
+{
+	GENERATED_USTRUCT_BODY()
+
+	// Store prev state to allow blending of prev->current state in material
+	// Uncomment this if we start blending states
+	//UPROPERTY(EditAnywhere, BlueprintReadWrite, Category = "AnimToTexture")
+	//FAnimToTextureAnimState PrevState;
+
+	UPROPERTY(EditAnywhere, BlueprintReadWrite, Category = "AnimToTexture")
+	FAnimToTextureAnimState CurrentState;
+};
+
+USTRUCT(BlueprintType)
+struct FAnimToTextureAnimationSyncData
+{
+	GENERATED_USTRUCT_BODY()
+
+	// The time used for sync when transitioning from skeletal mesh to material animated static mesh.
+	// World real time at the time of the transition
+	float SyncTime;
+};
+
+USTRUCT(BlueprintType)
+struct FAnimToTextureInstanceData
+{
+	GENERATED_USTRUCT_BODY()
+
+	UPROPERTY()
+	TArray<FAnimToTextureInstancePlaybackData> PlaybackData;
+
+	UPROPERTY()
+	TArray<FInstancedStaticMeshInstanceData> StaticMeshInstanceData;
+};
+
+class UAnimToTextureDataAsset;
+
+UCLASS()
+class ANIMTOTEXTURE_API UAnimToTextureInstancePlaybackLibrary : public UBlueprintFunctionLibrary
+{
+	GENERATED_BODY()
+
+public:
+
+	UFUNCTION(BlueprintCallable, Category = "Anim Texture Playback")
+	static void SetupInstancedMeshComponent(UInstancedStaticMeshComponent* InstancedMeshComponent, UPARAM(ref) FAnimToTextureInstanceData& InstanceData, int32 NumInstances);
+
+	UFUNCTION(BlueprintCallable, Category = "Anim Texture Playback")
+	static void BatchUpdateInstancedMeshComponent(UInstancedStaticMeshComponent* InstancedMeshComponent, UPARAM(ref) FAnimToTextureInstanceData& InstanceData);
+
+	UFUNCTION(BlueprintCallable, Category = "Anim Texture Playback")
+	static void AllocateInstanceData(UPARAM(ref) FAnimToTextureInstanceData& InstanceData, int32 Count);
+
+	UFUNCTION(BlueprintCallable, Category = "Anim Texture Playback")
+	static bool UpdateInstanceData(UPARAM(ref) FAnimToTextureInstanceData& InstanceData, int32 InstanceIndex, const FAnimToTextureInstancePlaybackData& PlaybackData, const FTransform& Transform);
+
+	UFUNCTION(BlueprintCallable, Category = "Anim Texture Playback")
+	static bool GetInstancePlaybackData(UPARAM(ref) const FAnimToTextureInstanceData& InstanceData, int32 InstanceIndex, FAnimToTextureInstancePlaybackData& InstancePlaybackData);
+
+	UFUNCTION(BlueprintCallable, Category = "Anim Texture Playback")
+	static bool GetInstanceTransform(UPARAM(ref) const FAnimToTextureInstanceData& InstanceData, int32 InstanceIndex, FTransform& InstanceTransform);
+
+	UFUNCTION(BlueprintPure, BlueprintCallable, Category = "Anim Texture Playback")
+	static bool AnimStateFromDataAsset(const UAnimToTextureDataAsset* DataAsset, int32 StateIndex, FAnimToTextureAnimState& AnimState);
+};
diff --git a/Engine/Plugins/Animation/AnimToTexture/Source/AnimToTexture/Public/EvaluateSequenceAnimInstance.h b/Engine/Plugins/Animation/AnimToTexture/Source/AnimToTexture/Public/EvaluateSequenceAnimInstance.h
new file mode 100644
--- /dev/null
+++ b/Engine/Plugins/Animation/AnimToTexture/Source/AnimToTexture/Public/EvaluateSequenceAnimInstance.h
@@ -0,0 +1,27 @@
+// Copyright Epic Games, Inc. All Rights Reserved.
+
+#pragma once
+
+#include "Animation/AnimInstance.h"
+#include "Animation/AnimSequence.h"
+#include "EvaluateSequenceAnimInstance.generated.h"
+
+/**
+ * Simple anim instance with a few parameters to be used as inputs when driving it
+ * from UAnimToTextureBPLibrary::AnimationToTexture()
+ */
+UCLASS()
+class ANIMTOTEXTURE_API UEvaluateSequenceAnimInstance : public UAnimInstance
+{
+	GENERATED_BODY()
+
+public:
+	UEvaluateSequenceAnimInstance(const FObjectInitializer& ObjectInitializer = FObjectInitializer::Get());
+
+	UPROPERTY(Transient, EditAnywhere, Category = Default, BlueprintReadOnly)
+	UAnimSequence* SequenceToEvaluate = nullptr;
+
+	UPROPERTY(Transient, EditAnywhere, Category = Default, BlueprintReadOnly)
+	float TimeToEvaluate = 0.0f;
+};
+
diff --git a/Engine/Plugins/Animation/AnimToTexture/Source/AnimToTexture/Public/LightweightMontageInstance.h b/Engine/Plugins/Animation/AnimToTexture/Source/AnimToTexture/Public/LightweightMontageInstance.h
new file mode 100644
--- /dev/null
+++ b/Engine/Plugins/Animation/AnimToTexture/Source/AnimToTexture/Public/LightweightMontageInstance.h
@@ -0,0 +1,95 @@
+// Copyright Epic Games, Inc. All Rights Reserved.
+
+#pragma once
+
+#include "UObject/WeakObjectPtrTemplates.h"
+#include "Animation/AnimationAsset.h"
+
+class UAnimMontage;
+class UAnimSequence;
+struct FLightweightMontageInstance;
+
+enum class EMontageSubStepResult : uint8;
+
+namespace UE
+{
+	namespace VertexAnimation
+	{
+		struct FLightWeightMontageExtractionSettings
+		{
+			bool bExtractRootMotion = false;
+		};
+
+		struct FLightweightMontageInstance;
+
+		struct FLightweightMontageSubstepper
+		{
+			void AddEvaluationTime(float InDeltaTime) { TimeRemaining += InDeltaTime; }
+			bool HasTimeRemaining() const { return (TimeRemaining > SMALL_NUMBER); }
+			bool HasReachedEndOfSection() const { return bReachedEndOfSection; }
+			float GetRemainingTime() const { return TimeRemaining; }
+			EMontageSubStepResult Advance(float& InOutTime, UAnimMontage* Montage);
+			int32 GetCurrentSectionIndex() const { return CurrentSectionIndex; }
+			float GetCurrentSectionLength() const { return CurrentSectionLength; }
+			float GetCurrentSectionStartTime() const { return CurrentSectionStartTime; }
+			float GetCurrentDeltaMove() const { return DeltaMove; }
+
+			void Initialize(UE::VertexAnimation::FLightweightMontageInstance* InMontageInstance)
+			{
+			}
+
+			float TimeRemaining = 0.0f;
+			float DeltaMove = 0.0f;
+			float CurrentSectionLength = 0.0f;
+			float CurrentSectionStartTime = 0.0f;
+			int32 CurrentSectionIndex = 0;
+			bool bReachedEndOfSection = false;
+		};
+
+		struct ANIMTOTEXTURE_API FLightweightMontageInstance
+		{
+			bool SequenceChangedThisFrame() const { return bSequenceChangedThisFrame; }
+			bool WasInitialized() const { return bInitialized; }
+			bool IsValid() const;
+			void Advance(float DeltaTime, float InGlobalTime, struct FRootMotionMovementParams& OutRootMotionParams, const FLightWeightMontageExtractionSettings& ExtractionSettings);
+			void Initialize(UAnimMontage* InMontage, float StartTime = 0.0f);
+			UAnimMontage* GetMontage() const;
+			const UAnimSequence* GetSequence() const;
+			float GetPositionInSection() const;
+			float GetPosition() const { return Position; }
+			float GetLength() const { return CachedLength; }
+			void Terminate();
+
+		private:
+			bool Advance_Internal(float DeltaTime, struct FRootMotionMovementParams& OutRootMotionParams, const FLightWeightMontageExtractionSettings& ExtractionSettings);
+			void RefreshNextPrevSections(UAnimMontage* InMontage);
+			void SetSequence(const UAnimSequence* InSequence);
+			const UAnimSequence* FindSequenceAtCurrentTime(float& SequenceStartTime);
+
+		public:
+			
+			FRootMotionMovementParams RootMotionDelta;
+
+		private:
+
+			UE::VertexAnimation::FLightweightMontageSubstepper MontageSubStepper;
+
+			// @todo : Handle asset references externally to avoid paying the cost of the weak ptr
+			TWeakObjectPtr<UAnimMontage> MontageWeak = nullptr;
+			TWeakObjectPtr<const UAnimSequence> SequenceWeak = nullptr;
+
+			// Transient. Gets cleared at the end of update
+			UAnimMontage* Montage = nullptr;
+
+			TArray<int32, TInlineAllocator<4>> NextSections;
+			float Position = 0.0f;
+			float CachedLength = 0.0f;
+
+			// We need this since GetSegmentIndexAtTime and GetAnimCompositeSectionIndexFromPos can return different sections at the same frame
+			float SectionStartTimeTransient = 0.0f;
+
+			bool bSequenceChangedThisFrame = false;
+			bool bInitialized = false;
+		};
+	}
+}
\ No newline at end of file
diff --git a/Engine/Plugins/Animation/AnimToTexture/Source/AnimToTextureEditor/AnimToTextureEditor.Build.cs b/Engine/Plugins/Animation/AnimToTexture/Source/AnimToTextureEditor/AnimToTextureEditor.Build.cs
new file mode 100644
--- /dev/null
+++ b/Engine/Plugins/Animation/AnimToTexture/Source/AnimToTextureEditor/AnimToTextureEditor.Build.cs
@@ -0,0 +1,57 @@
+// Copyright Epic Games, Inc. All Rights Reserved.
+
+using UnrealBuildTool;
+
+public class AnimToTextureEditor : ModuleRules
+{
+	public AnimToTextureEditor(ReadOnlyTargetRules Target) : base(Target)
+	{
+		PCHUsage = ModuleRules.PCHUsageMode.UseExplicitOrSharedPCHs;
+		
+		PublicIncludePaths.AddRange(
+			new string[] {
+				// ... add public include paths required here ...
+			}
+			);
+				
+		
+		PrivateIncludePaths.AddRange(
+			new string[] {
+				// ... add other private include paths required here ...
+			}
+			);
+			
+		
+		PublicDependencyModuleNames.AddRange(
+			new string[]
+			{
+				"Core",
+				"MaterialEditor"
+				// ... add other public dependencies that you statically link with here ...
+			}
+			);
+			
+		
+		PrivateDependencyModuleNames.AddRange(
+			new string[]
+			{
+				"CoreUObject",
+				"Engine",
+				"Slate",
+				"SlateCore",
+				"AnimToTexture",
+				"RawMesh",
+				"MeshDescription",
+				"UnrealEd",
+			}
+			);
+		
+		
+		DynamicallyLoadedModuleNames.AddRange(
+			new string[]
+			{
+				// ... add any modules that your module loads dynamically here ...
+			}
+			);
+	}
+}
diff --git a/Engine/Plugins/Animation/AnimToTexture/Source/AnimToTextureEditor/Private/AnimToTextureBPLibrary.cpp b/Engine/Plugins/Animation/AnimToTexture/Source/AnimToTextureEditor/Private/AnimToTextureBPLibrary.cpp
new file mode 100644
--- /dev/null
+++ b/Engine/Plugins/Animation/AnimToTexture/Source/AnimToTextureEditor/Private/AnimToTextureBPLibrary.cpp
@@ -0,0 +1,998 @@
+// Copyright Epic Games, Inc. All Rights Reserved.
+
+#include "AnimToTextureBPLibrary.h"
+#include "AnimToTexture.h"
+#include "AnimToTextureUtils.h"
+#include "AnimToTextureSkeletalMesh.h"
+#include "EvaluateSequenceAnimInstance.h"
+
+#include "Editor.h"
+#include "LevelEditor.h"
+#include "RawMesh.h"
+#include "MeshUtilities.h"
+#include "AssetRegistry/AssetRegistryModule.h"
+#include "Components/SkeletalMeshComponent.h"
+#include "Components/SkinnedMeshComponent.h"
+#include "Rendering/SkeletalMeshRenderData.h"
+#include "Rendering/SkeletalMeshModel.h"
+#include "Animation/AnimBlueprint.h"
+#include "Animation/AnimSequence.h"
+#include "Math/Vector.h"
+#include "Math/NumericLimits.h"
+#include "Engine/SkeletalMeshSocket.h"
+#include "MeshDescription.h"
+#include "Materials/MaterialInstanceConstant.h"
+#include "MaterialEditingLibrary.h"
+#include "EditorSupportDelegates.h"
+
+UAnimToTextureBPLibrary::UAnimToTextureBPLibrary(const FObjectInitializer& ObjectInitializer)
+: Super(ObjectInitializer)
+{
+
+}
+
+void UAnimToTextureBPLibrary::AnimationToTexture(UAnimToTextureDataAsset* DataAsset,
+	const FTransform RootTransform, const bool bCreateTextures, const bool bCreateUVChannel)
+{
+	if (!DataAsset)
+	{
+		return;
+	}
+
+	// Reset DataAsset Values
+	DataAsset->Reset();
+
+	if (!DataAsset->GetSkeletalMesh() || !DataAsset->GetStaticMesh())
+	{
+		UE_LOG(LogTemp, Warning, TEXT("Invalid Meshes"));
+		return;
+	}
+
+	// Check LODs
+	if (!DataAsset->GetSkeletalMesh()->IsValidLODIndex(DataAsset->SkeletalLODIndex) || 
+		DataAsset->StaticLODIndex >= DataAsset->GetStaticMesh()->GetNumLODs())
+	{
+		UE_LOG(LogTemp, Warning, TEXT("Invalid LOD Indices"));
+		return;
+	}
+
+	// Check Socket.
+	bool bValidSocket = false;
+	if (DataAsset->AttachToSocket.IsValid() && !DataAsset->AttachToSocket.IsNone())
+	{
+		if (AnimToTexture_Private::HasBone(DataAsset->GetSkeletalMesh(), DataAsset->AttachToSocket))
+		{
+			bValidSocket = true;
+		}
+		else
+		{
+			UE_LOG(LogTemp, Warning, TEXT("Invalid Socket: %s"), *DataAsset->AttachToSocket.ToString());
+			return;
+		}
+	}
+	
+	if (bValidSocket && DataAsset->Mode == EAnimToTextureMode::Vertex)
+	{
+		UE_LOG(LogTemp, Warning, TEXT("Unable to use Socket in Vertex Mode. Use Bone Mode instead."));
+		return;
+	}
+
+	// ---------------------------------------------------------------------------		
+	// Get Meshes Vertices and Mapping.
+	// NOTE: We need to create a Mapping between the StaticMesh|MeshDescription and the SkeletalMesh
+	//       Since they dont have same number of points.
+	//	
+	TArray<FVector3f> Vertices;
+	TArray<FVector3f> SkelVertices;
+	TArray<int32> StaticToSkelMapping;
+
+	AnimToTexture_Private::GetStaticToSkeletalMapping(
+		*DataAsset->GetStaticMesh(), DataAsset->StaticLODIndex,
+		*DataAsset->GetSkeletalMesh(), DataAsset->SkeletalLODIndex,
+		Vertices, SkelVertices,
+		StaticToSkelMapping);
+	
+	// Set Static Number of Vertices 
+	// NOTE: these are the MeshDescription Vertices.
+	const int32 NumVertices = Vertices.Num();
+
+	// --------------------------------------------------------------------------
+
+	// Create Temp Actor
+	UWorld* World = GEditor->GetEditorWorldContext().World();
+	AActor* Actor = World->SpawnActor<AActor>();
+
+	// Create Temp LeaderMeshComponent
+	USkeletalMeshComponent* LeaderMeshComponent = nullptr;
+	if (DataAsset->GetLeaderSkeletalMesh())
+	{
+		LeaderMeshComponent = NewObject<USkeletalMeshComponent>(Actor);
+		LeaderMeshComponent->SetSkeletalMesh(DataAsset->GetLeaderSkeletalMesh());
+		LeaderMeshComponent->SetForcedLOD(1); // Force to LOD0
+		LeaderMeshComponent->SetAnimationMode(EAnimationMode::AnimationSingleNode);
+		LeaderMeshComponent->SetUpdateAnimationInEditor(true);
+		LeaderMeshComponent->VisibilityBasedAnimTickOption = EVisibilityBasedAnimTickOption::AlwaysTickPoseAndRefreshBones;		
+		LeaderMeshComponent->RegisterComponent();
+	};
+
+	// Create Temp SkeletalMesh Component
+	USkeletalMeshComponent* MeshComponent = NewObject<USkeletalMeshComponent>(Actor);
+	MeshComponent->SetSkeletalMesh(DataAsset->GetSkeletalMesh());
+	MeshComponent->SetForcedLOD(1); // Force to LOD0;
+	MeshComponent->RegisterComponent();
+	
+	USkeletalMeshComponent* BaseMeshComponent = nullptr;
+	if (LeaderMeshComponent)
+	{
+		MeshComponent->SetupAttachment(LeaderMeshComponent);
+		MeshComponent->SetLeaderPoseComponent(LeaderMeshComponent);
+
+		BaseMeshComponent = LeaderMeshComponent;
+	}
+	else
+	{
+		MeshComponent->SetAnimationMode(EAnimationMode::AnimationSingleNode);
+		MeshComponent->SetUpdateAnimationInEditor(true);
+		MeshComponent->VisibilityBasedAnimTickOption = EVisibilityBasedAnimTickOption::AlwaysTickPoseAndRefreshBones;
+		
+		BaseMeshComponent = MeshComponent;
+	}
+	
+	// --------------------------------------------------------------------------
+	// Set AnimBlueprint
+	UEvaluateSequenceAnimInstance* EvaluationAnimInstance = nullptr;
+	if (DataAsset->AnimInstanceClass != nullptr)
+	{
+		// Set Anim Blueprint
+		BaseMeshComponent->SetAnimInstanceClass(DataAsset->AnimInstanceClass);
+		EvaluationAnimInstance = Cast<UEvaluateSequenceAnimInstance>(BaseMeshComponent->GetAnimInstance());
+	}
+
+	// Get Mapping between Follower and Leader
+	TArray<int32> FollowerBoneMap;
+	if (DataAsset->GetLeaderSkeletalMesh() &&
+		DataAsset->Mode == EAnimToTextureMode::Bone)
+	{
+		AnimToTexture_Private::GetFollowerBoneMap(DataAsset->GetLeaderSkeletalMesh(), DataAsset->GetSkeletalMesh(), FollowerBoneMap);
+	}
+	
+	// ---------------------------------------------------------------------------
+	// Get Reference Skeleton Transforms
+	//
+	int32 NumBones = INDEX_NONE;
+	int32 SocketIndex = INDEX_NONE;
+	TArray<FName>     BoneNames;
+	TArray<FVector3f> BoneRefPositions;
+	TArray<FVector4>  BoneRefRotations;
+	TArray<FVector3f> BonePositions;
+	TArray<FVector4>  BoneRotations;
+	
+	if (DataAsset->Mode == EAnimToTextureMode::Bone)
+	{
+		// Get Number of RawBones (no virtual)
+		NumBones = AnimToTexture_Private::GetNumBones(BaseMeshComponent->GetSkeletalMesh());
+
+		// Get Raw Ref Bone (no virtual)
+		TArray<FTransform> RefBoneTransforms;
+		AnimToTexture_Private::GetRefBoneTransforms(BaseMeshComponent->GetSkeletalMesh(), RefBoneTransforms);
+		AnimToTexture_Private::DecomposeTransformations(RefBoneTransforms, BoneRefPositions, BoneRefRotations);
+
+		// Get Bone Names (no virtual)
+		AnimToTexture_Private::GetBoneNames(BaseMeshComponent->GetSkeletalMesh(), BoneNames);
+
+		// Make sure array sizes are correct.
+		check(RefBoneTransforms.Num() == NumBones && BoneNames.Num() == NumBones);
+
+		// Check if Socket is in BoneNames
+		if (bValidSocket && !BoneNames.Find(DataAsset->AttachToSocket, SocketIndex))
+		{
+			UE_LOG(LogTemp, Warning, TEXT("Socket: %s not found in Raw Bone List"), *DataAsset->AttachToSocket.ToString());
+			return;
+		}
+
+		// Add RefPose 
+		// Note: this is added in the first frame of the Bone Position and Rotation Textures
+		BonePositions.Append(BoneRefPositions);
+		BoneRotations.Append(BoneRefRotations);
+	}
+
+	// ---------------------------------------------------------------------------
+	// Get Vertex Data (for all frames)
+	//		
+	TArray<FVector3f> VertexDeltas;
+	TArray<FVector3f> VertexNormals;
+	
+	// Get Animation Frames Data
+	//
+	const float SampleInterval = 1.f / DataAsset->SampleRate;
+
+	for (const FAnimSequenceInfo& AnimSequenceInfo : DataAsset->AnimSequences)
+	{
+		UAnimSequence* AnimSequence = AnimSequenceInfo.AnimSequence;
+
+		if (!AnimSequenceInfo.bEnabled || !AnimSequence)
+		{
+			continue;
+		}
+		
+		// Make sure SkelMesh and AnimSequence use same Skeleton
+		if (AnimSequence->GetSkeleton() != BaseMeshComponent->GetSkeletalMesh()->GetSkeleton())
+		{
+			UE_LOG(LogTemp, Warning, TEXT("Invalid AnimSequence: %s for given SkeletalMesh: %s"), *AnimSequence->GetFName().ToString(), *BaseMeshComponent->GetSkeletalMesh()->GetFName().ToString());
+			continue;
+		}
+		// Set AnimSequence
+		else
+		{
+			if (IsValid(EvaluationAnimInstance))
+			{
+				EvaluationAnimInstance->SequenceToEvaluate = AnimSequence;
+				EvaluationAnimInstance->TimeToEvaluate = 0.0f;
+			}
+			else
+			{
+				BaseMeshComponent->SetAnimation(AnimSequence);
+			}
+		}
+
+		// -----------------------------------------------------------------------------------
+		// Get Number of Frames
+		//
+		int32 StartFrame;
+		int32 EndFrame;
+		
+		// Get Range from AnimSequence
+		if (!AnimSequenceInfo.bUseCustomRange)
+		{
+			StartFrame = 0;
+			EndFrame = AnimSequence->GetNumberOfSampledKeys() - 1; // AnimSequence->GetNumberOfFrames();
+		}
+		// Get Range from DataAsset
+		else
+		{
+			StartFrame = AnimSequenceInfo.StartFrame;
+			EndFrame = AnimSequenceInfo.EndFrame;
+		}
+		if (EndFrame - StartFrame <= 0)
+		{
+			continue;
+		}
+		
+		// ---------------------------------------------------------------------------
+		// 
+		float Time = 0.f;
+		float EndTime = AnimSequence->GetTimeAtFrame(EndFrame);
+
+		int32 SampleIndex = 0;
+
+		while (Time < EndTime)
+		{
+			Time = FMath::Clamp(SampleIndex * SampleInterval, 0.f, EndTime);
+			SampleIndex++;
+
+			// Go To Time
+			if (IsValid(EvaluationAnimInstance))
+			{
+				EvaluationAnimInstance->TimeToEvaluate = Time;
+			}
+			else
+			{
+				BaseMeshComponent->SetPosition(Time);
+			}
+
+			// Update SkelMesh Animation.
+			BaseMeshComponent->TickAnimation(0.0f, false /*bNeedsValidRootMotion*/);
+			// BaseMeshComponent->TickComponent(0.0f, ELevelTick::LEVELTICK_All, nullptr);
+			BaseMeshComponent->RefreshBoneTransforms(nullptr /*TickFunction*/);
+			BaseMeshComponent->RefreshFollowerComponents();
+			
+			// ---------------------------------------------------------------------------
+			// Store Vertex Deltas & Normals.
+			//
+			if (DataAsset->Mode == EAnimToTextureMode::Vertex)
+			{
+				TArray<FVector3f> Positions;
+				TArray<FVector3f> Normals;
+				AnimToTexture_Private::GetSkinnedVertices(MeshComponent, DataAsset->SkeletalLODIndex,
+					Positions, Normals);
+
+				// Loop thru static vertices and find the mapped SkeletalMesh Vertex
+				for (int32 VertexIndex = 0; VertexIndex < NumVertices; ++VertexIndex)
+				{
+					const int32 SkelVertexIndex = StaticToSkelMapping[VertexIndex];
+
+					// Delta & Normal
+					const FVector3f VertexDelta = (FVector3f)RootTransform.TransformPosition((FVector)Positions[SkelVertexIndex]) - SkelVertices[SkelVertexIndex];
+					const FVector3f VertexNormal = (FVector3f)RootTransform.TransformVector((FVector)Normals[SkelVertexIndex]);
+
+					VertexDeltas.Add(VertexDelta);
+					VertexNormals.Add(VertexNormal);
+				}
+			} // End Vertex Mode
+
+			// ---------------------------------------------------------------------------
+			// Store Bone Positions & Rotations
+			//
+			else if (DataAsset->Mode == EAnimToTextureMode::Bone)
+			{
+				// Get Relative Transforms
+				// Note: Size is of Raw bones in SkeletalMesh. These are the original/raw bones of the asset, without Virtual Bones.
+				// Note: You cannot call CacheRefToLocalMatrices on a FollowerComponent
+				TArray<FMatrix44f> RefToLocals;
+				BaseMeshComponent->CacheRefToLocalMatrices(RefToLocals);
+				
+				// check size
+				check(RefToLocals.Num() == NumBones);
+
+				// Get Component Space Transforms
+				// Note returns all transforms, including VirtualBones
+				// Note: You cannot call GetComponentSpaceTransforms on a FollowerComponent
+				const TArray<FTransform>& CompSpaceTransforms = BaseMeshComponent->GetComponentSpaceTransforms();
+				check(CompSpaceTransforms.Num() >= RefToLocals.Num());
+
+				for (int32 BoneIndex = 0; BoneIndex < NumBones; ++BoneIndex)
+				{
+					// Decompose Transformation (ComponentSpace)
+					const FTransform& CompSpaceTransform = CompSpaceTransforms[BoneIndex];
+					FVector3f BonePosition;
+					FVector4 BoneRotation;
+					AnimToTexture_Private::DecomposeTransformation(CompSpaceTransform, BonePosition, BoneRotation);
+
+					// Position Delta (from RefPose)
+					const FVector3f Delta = BonePosition - BoneRefPositions[BoneIndex];
+
+					// Decompose Transformation (Relative to RefPose)
+					FVector3f BoneRelativePosition;
+					FVector4 BoneRelativeRotation;
+					FMatrix RefToLocalMatrix = FMatrix(RefToLocals[BoneIndex]);
+					const FTransform RelativeTransform(RefToLocalMatrix);
+					AnimToTexture_Private::DecomposeTransformation(RelativeTransform, BoneRelativePosition, BoneRelativeRotation);
+
+					BonePositions.Add(Delta);
+					BoneRotations.Add(BoneRelativeRotation);
+				}
+			} // End Bone Mode
+		} // End Frame
+
+		// Store Anim Info Data
+		FAnimInfo AnimInfo;
+		AnimInfo.NumFrames = SampleIndex;
+		AnimInfo.AnimStart = DataAsset->NumFrames;
+		AnimInfo.bLooping = AnimSequenceInfo.bLooping;
+		DataAsset->Animations.Add(AnimInfo);
+
+		// Accumulate Frames
+		DataAsset->NumFrames += AnimInfo.NumFrames;
+	} // End Anim
+		
+	// Destroy Temp Component
+	MeshComponent->UnregisterComponent();
+	MeshComponent->DestroyComponent();
+
+	// Destroy Leader Component
+	if (LeaderMeshComponent)
+	{
+		LeaderMeshComponent->UnregisterComponent();
+		LeaderMeshComponent->DestroyComponent();
+	}
+	
+	Actor->Destroy();
+	
+	// ---------------------------------------------------------------------------
+	// Nothing to do here ...
+	//
+	if (!DataAsset->NumFrames || !NumVertices)
+	{
+		return;
+	}
+
+	// ---------------------------------------------------------------------------
+	if (DataAsset->Mode == EAnimToTextureMode::Vertex)
+	{
+		// Find Best Resolution for Vertex Data
+		int32 Height, Width;
+		if (!FindBestResolution(DataAsset->NumFrames, NumVertices, 
+								Height, Width, DataAsset->VertexRowsPerFrame, 
+								DataAsset->MaxHeight, DataAsset->MaxWidth, DataAsset->bEnforcePowerOfTwo))
+		{
+			UE_LOG(LogTemp, Warning, TEXT("Vertex Animation data cannot be fit in a %ix%i texture."), DataAsset->MaxHeight, DataAsset->MaxWidth);
+			return;
+		}
+
+		// Normalize Vertex Data
+		TArray<FVector3f> NormalizedVertexDeltas;
+		TArray<FVector3f> NormalizedVertexNormals;
+		NormalizeVertexData(
+			VertexDeltas, VertexNormals,
+			DataAsset->VertexMinBBox, DataAsset->VertexSizeBBox,
+			NormalizedVertexDeltas, NormalizedVertexNormals);
+
+		// Write Textures
+		if (bCreateTextures)
+		{
+			AnimToTexture_Private::WriteVectorsToTexture<FVector3f, AnimToTexture_Private::FLowPrecision>(NormalizedVertexDeltas, DataAsset->NumFrames, DataAsset->VertexRowsPerFrame, Height, Width, DataAsset->GetVertexPositionTexture());
+			AnimToTexture_Private::WriteVectorsToTexture<FVector3f, AnimToTexture_Private::FLowPrecision>(NormalizedVertexNormals, DataAsset->NumFrames, DataAsset->VertexRowsPerFrame, Height, Width, DataAsset->GetVertexNormalTexture());
+		}
+
+		// Add Vertex UVChannel
+		if (bCreateUVChannel)
+		{
+			CreateUVChannel(DataAsset->GetStaticMesh(), DataAsset->StaticLODIndex, DataAsset->UVChannel,
+				Height, Width);
+		}
+	}
+
+	// ---------------------------------------------------------------------------
+	
+	if (DataAsset->Mode == EAnimToTextureMode::Bone)
+	{
+		// Find Best Resolution for Bone Data
+		int32 Height, Width;
+
+		// NOTE: If NumBones are > 256, you might need to use LeaderPose
+
+		// Note we are adding +1 frame for the ref pose
+		if (!FindBestResolution(DataAsset->NumFrames + 1, NumBones, 
+								Height, Width, DataAsset->BoneRowsPerFrame, 
+								DataAsset->MaxHeight, DataAsset->MaxWidth, DataAsset->bEnforcePowerOfTwo))
+		{
+			UE_LOG(LogTemp, Warning, TEXT("Bone Animation data cannot be fit in a %ix%i texture."), DataAsset->MaxHeight, DataAsset->MaxWidth);
+			return;
+		}
+
+		// Write Bone Position and Rotation Textures
+		if (bCreateTextures)
+		{
+			// Normalize Bone Data
+			TArray<FVector3f> NormalizedBonePositions;
+			TArray<FVector4> NormalizedBoneRotations;
+			NormalizeBoneData(
+				BonePositions, BoneRotations,
+				DataAsset->BoneMinBBox, DataAsset->BoneSizeBBox,
+				NormalizedBonePositions, NormalizedBoneRotations);
+
+			// Write Textures
+			if (DataAsset->PositionAndRotationPrecision == EAnimToTextureBonePrecision::SixteenBits)
+			{
+				AnimToTexture_Private::WriteVectorsToTexture<FVector3f, AnimToTexture_Private::FHighPrecision>(NormalizedBonePositions, DataAsset->NumFrames + 1, DataAsset->BoneRowsPerFrame, Height, Width, DataAsset->GetBonePositionTexture());
+				AnimToTexture_Private::WriteVectorsToTexture<FVector4, AnimToTexture_Private::FHighPrecision>(NormalizedBoneRotations, DataAsset->NumFrames + 1, DataAsset->BoneRowsPerFrame, Height, Width, DataAsset->GetBoneRotationTexture());
+			}
+			else
+			{
+				AnimToTexture_Private::WriteVectorsToTexture<FVector3f, AnimToTexture_Private::FLowPrecision>(NormalizedBonePositions, DataAsset->NumFrames + 1, DataAsset->BoneRowsPerFrame, Height, Width, DataAsset->GetBonePositionTexture());
+				AnimToTexture_Private::WriteVectorsToTexture<FVector4, AnimToTexture_Private::FLowPrecision>(NormalizedBoneRotations, DataAsset->NumFrames + 1, DataAsset->BoneRowsPerFrame, Height, Width, DataAsset->GetBoneRotationTexture());
+			}
+		}
+
+		// ---------------------------------------------------------------------------
+
+		// Find Best Resolution for Weights Data
+		if (!FindBestResolution(2, NumVertices, 
+								Height, Width, DataAsset->BoneWeightRowsPerFrame, 
+								DataAsset->MaxHeight, DataAsset->MaxWidth, DataAsset->bEnforcePowerOfTwo))
+		{
+			UE_LOG(LogTemp, Warning, TEXT("Weights Data cannot be fit in a %ix%i texture."), DataAsset->MaxHeight, DataAsset->MaxWidth);
+			return;
+		}
+
+		// Write Weights Texture
+		if (bCreateTextures)
+		{
+			TArray<AnimToTexture_Private::TVertexSkinWeight<4>> SkinWeights;
+
+			// Store Influence Weights
+			if (!bValidSocket)
+			{
+				// Get SkinWeights mapped to StaticMesh
+				AnimToTexture_Private::GetReducedSkinWeightsData(
+					*DataAsset->GetSkeletalMesh(), DataAsset->SkeletalLODIndex,
+					StaticToSkelMapping, 
+					SkinWeights);
+
+				// Remap MeshBoneIndices to FollowerBoneMap
+				if (DataAsset->GetLeaderSkeletalMesh())
+				{
+					for (AnimToTexture_Private::TVertexSkinWeight<4>& SkinWeight : SkinWeights)
+					{
+						SkinWeight.MeshBoneIndices[0] = FollowerBoneMap[SkinWeight.MeshBoneIndices[0]];
+						SkinWeight.MeshBoneIndices[1] = FollowerBoneMap[SkinWeight.MeshBoneIndices[1]];
+						SkinWeight.MeshBoneIndices[2] = FollowerBoneMap[SkinWeight.MeshBoneIndices[2]];
+						SkinWeight.MeshBoneIndices[3] = FollowerBoneMap[SkinWeight.MeshBoneIndices[3]];
+					}
+				}
+			}
+			// If Valid Socket, set all influences to same index.
+			else
+			{
+				// Set all indices and weights to same SocketIndex
+				SkinWeights.SetNumUninitialized(NumVertices);
+				for (AnimToTexture_Private::TVertexSkinWeight<4>& SkinWeight : SkinWeights)
+				{
+					SkinWeight.BoneWeights = TStaticArray<uint8, 4>(InPlace, 255);
+					SkinWeight.MeshBoneIndices = TStaticArray<uint16, 4>(InPlace, SocketIndex);
+				}
+			}
+
+			// Write Bone Weights Texture
+			AnimToTexture_Private::WriteSkinWeightsToTexture(SkinWeights,
+				DataAsset->BoneWeightRowsPerFrame, Height, Width, DataAsset->GetBoneWeightTexture());
+		}
+
+		// Add Vertex UVChannel
+		if (bCreateUVChannel)
+		{
+			CreateUVChannel(DataAsset->GetStaticMesh(), DataAsset->StaticLODIndex, DataAsset->UVChannel,
+				Height, Width);
+		}
+	}
+
+	// ---------------------------------------------------------------------------
+	// Mark DataAsset as Dirty.
+	//
+	DataAsset->MarkPackageDirty();
+}
+
+
+void UAnimToTextureBPLibrary::UpdateMaterialInstanceFromDataAsset(UAnimToTextureDataAsset* DataAsset, UMaterialInstanceConstant* MaterialInstance, 
+	const FAnimToTextureMaterialParamNames& ParamNames, const EMaterialParameterAssociation MaterialParameterAssociation)
+{
+	if (!MaterialInstance || !DataAsset)
+	{
+		return;
+	}
+	
+	// Set Preview Mesh
+	if (DataAsset->GetStaticMesh())
+	{
+		MaterialInstance->PreviewMesh = DataAsset->GetStaticMesh();
+	}
+
+	if (DataAsset->Mode == EAnimToTextureMode::Vertex)
+	{
+		FLinearColor VectorParameter;
+		VectorParameter = FLinearColor(DataAsset->VertexMinBBox);
+		UMaterialEditingLibrary::SetMaterialInstanceVectorParameterValue(MaterialInstance, ParamNames.BoundingBoxMin, VectorParameter, MaterialParameterAssociation);
+		
+		VectorParameter = FLinearColor(DataAsset->VertexSizeBBox);
+		UMaterialEditingLibrary::SetMaterialInstanceVectorParameterValue(MaterialInstance, ParamNames.BoundingBoxScale, VectorParameter, MaterialParameterAssociation);
+		UMaterialEditingLibrary::SetMaterialInstanceScalarParameterValue(MaterialInstance, ParamNames.NumFrames, DataAsset->NumFrames, MaterialParameterAssociation);
+		UMaterialEditingLibrary::SetMaterialInstanceScalarParameterValue(MaterialInstance, ParamNames.RowsPerFrame, DataAsset->VertexRowsPerFrame, MaterialParameterAssociation);
+
+		UMaterialEditingLibrary::SetMaterialInstanceTextureParameterValue(MaterialInstance, ParamNames.VertexPositionTexture, DataAsset->GetVertexPositionTexture(), MaterialParameterAssociation);
+		UMaterialEditingLibrary::SetMaterialInstanceTextureParameterValue(MaterialInstance, ParamNames.VertexNormalTexture, DataAsset->GetVertexNormalTexture(), MaterialParameterAssociation);
+
+	}
+	else if (DataAsset->Mode == EAnimToTextureMode::Bone)
+	{
+		FLinearColor VectorParameter;
+		VectorParameter = FLinearColor(DataAsset->BoneMinBBox);
+		UMaterialEditingLibrary::SetMaterialInstanceVectorParameterValue(MaterialInstance, ParamNames.BoundingBoxMin, VectorParameter, MaterialParameterAssociation);
+
+		VectorParameter = FLinearColor(DataAsset->BoneSizeBBox);
+		UMaterialEditingLibrary::SetMaterialInstanceVectorParameterValue(MaterialInstance, ParamNames.BoundingBoxScale, VectorParameter, MaterialParameterAssociation);
+		UMaterialEditingLibrary::SetMaterialInstanceScalarParameterValue(MaterialInstance, ParamNames.NumFrames, DataAsset->NumFrames, MaterialParameterAssociation);
+		UMaterialEditingLibrary::SetMaterialInstanceScalarParameterValue(MaterialInstance, ParamNames.RowsPerFrame, DataAsset->BoneRowsPerFrame, MaterialParameterAssociation);
+		UMaterialEditingLibrary::SetMaterialInstanceScalarParameterValue(MaterialInstance, ParamNames.BoneWeightRowsPerFrame, DataAsset->BoneWeightRowsPerFrame, MaterialParameterAssociation);
+
+		UMaterialEditingLibrary::SetMaterialInstanceTextureParameterValue(MaterialInstance, ParamNames.BonePositionTexture, DataAsset->GetBonePositionTexture(), MaterialParameterAssociation);
+		UMaterialEditingLibrary::SetMaterialInstanceTextureParameterValue(MaterialInstance, ParamNames.BoneRotationTexture, DataAsset->GetBoneRotationTexture(), MaterialParameterAssociation);
+		UMaterialEditingLibrary::SetMaterialInstanceTextureParameterValue(MaterialInstance, ParamNames.BoneWeightsTexture, DataAsset->GetBoneWeightTexture(), MaterialParameterAssociation);
+	}
+
+	// Update Material
+	UMaterialEditingLibrary::UpdateMaterialInstance(MaterialInstance);
+
+	// Rebuild Material
+	UMaterialEditingLibrary::RebuildMaterialInstanceEditors(MaterialInstance->GetMaterial());
+}
+
+void UAnimToTextureBPLibrary::UpdateMaterialLayerFunction(
+	UMaterialInstanceConstant* MaterialInstance,
+	UMaterialFunctionInterface* OldMaterialFunction,
+	UMaterialFunctionInterface* NewMaterialFunction)
+{
+	if (!IsValid(MaterialInstance))
+	{
+		return;
+	}
+
+	const FMaterialParameterInfo ParameterInfo(TEXT("BoneAnimation"), EMaterialParameterAssociation::GlobalParameter);
+
+	FMaterialLayersFunctions LayersValue;
+	FGuid TempGuid(0, 0, 0, 0);
+
+	// TODO: Ben.Ingram / Cesar.Castro verify this.
+	bool bMaterialLayerReplaced = false;
+	if (MaterialInstance->GetMaterialLayers(LayersValue))
+	{
+		for (int32 LayerIdx = 0; LayerIdx < LayersValue.Layers.Num(); ++LayerIdx)
+		{
+			TObjectPtr<UMaterialFunctionInterface>& MaterialFunction = LayersValue.Layers[LayerIdx];
+			if (MaterialFunction.Get() == OldMaterialFunction)
+			{
+				MaterialFunction = NewMaterialFunction;
+				LayersValue.EditorOnly.LayerLinkStates[LayerIdx] = EMaterialLayerLinkState::UnlinkedFromParent;
+				bMaterialLayerReplaced = true;
+				break;
+			}
+
+		}
+	}
+
+	if (bMaterialLayerReplaced)
+	{
+		// TODO: Ben.Ingram / Cesar.Castro verify this.
+		MaterialInstance->SetMaterialLayers(LayersValue);
+
+		// Update Material
+		UMaterialEditingLibrary::UpdateMaterialInstance(MaterialInstance);
+
+		// Rebuild Material
+		UMaterialEditingLibrary::RebuildMaterialInstanceEditors(MaterialInstance->GetMaterial());
+	}
+}
+
+
+UStaticMesh* UAnimToTextureBPLibrary::ConvertSkeletalMeshToStaticMesh(USkeletalMesh* SkeletalMesh, const FString PackageName, const int32 LODIndex)
+{
+	if (!SkeletalMesh || PackageName.IsEmpty())
+	{
+		return nullptr;
+	}
+
+	if (!FPackageName::IsValidObjectPath(PackageName))
+	{
+		return nullptr;
+	}
+
+	if (LODIndex >= 0 && !SkeletalMesh->IsValidLODIndex(LODIndex))
+	{
+		UE_LOG(LogTemp, Warning, TEXT("Invalid LODIndex: %i"), LODIndex);
+		return nullptr;
+	}
+
+	// Create Temp Actor
+	UWorld* World = GEditor->GetEditorWorldContext().World();
+	AActor* Actor = World->SpawnActor<AActor>();
+
+	// Create Temp SkeletalMesh Component
+	USkeletalMeshComponent* MeshComponent = NewObject<USkeletalMeshComponent>(Actor);
+	MeshComponent->RegisterComponent();
+	MeshComponent->SetSkeletalMesh(SkeletalMesh);
+	TArray<UMeshComponent*> MeshComponents = { MeshComponent };
+
+	UStaticMesh* OutStaticMesh = nullptr;
+	bool bGeneratedCorrectly = true;
+
+	// Create New StaticMesh
+	if (!FPackageName::DoesPackageExist(PackageName))
+	{
+		IMeshUtilities& MeshUtilities = FModuleManager::Get().LoadModuleChecked<IMeshUtilities>("MeshUtilities");
+		OutStaticMesh = MeshUtilities.ConvertMeshesToStaticMesh(MeshComponents, FTransform::Identity, PackageName);
+	}
+	// Update Existing StaticMesh
+	else
+	{
+		// Load Existing Mesh
+		OutStaticMesh = LoadObject<UStaticMesh>(nullptr, *PackageName);
+	}
+
+	if (OutStaticMesh)
+	{
+		// Create Temp Package.
+		// because 
+		UPackage* TransientPackage = GetTransientPackage();
+
+		// Create Temp Mesh.
+		IMeshUtilities& MeshUtilities = FModuleManager::Get().LoadModuleChecked<IMeshUtilities>("MeshUtilities");
+		UStaticMesh* TempMesh = MeshUtilities.ConvertMeshesToStaticMesh(MeshComponents, FTransform::Identity, TransientPackage->GetPathName());
+
+		// make sure transactional flag is on
+		TempMesh->SetFlags(RF_Transactional);
+
+		// Copy All LODs
+		if (LODIndex < 0)
+		{
+			const int32 NumSourceModels = TempMesh->GetNumSourceModels();
+			OutStaticMesh->SetNumSourceModels(NumSourceModels);
+
+			for (int32 Index = 0; Index < NumSourceModels; ++Index)
+			{
+				// Get RawMesh
+				FRawMesh RawMesh;
+				TempMesh->GetSourceModel(Index).LoadRawMesh(RawMesh);
+
+				// Set RawMesh
+				OutStaticMesh->GetSourceModel(Index).SaveRawMesh(RawMesh);
+			};
+		}
+
+		// Copy Single LOD
+		else
+		{
+			if (LODIndex >= TempMesh->GetNumSourceModels())
+			{
+				UE_LOG(LogTemp, Warning, TEXT("Invalid Source Model Index: %i"), LODIndex);
+				bGeneratedCorrectly = false;
+			}
+			else
+			{
+				OutStaticMesh->SetNumSourceModels(1);
+
+				// Get RawMesh
+				FRawMesh RawMesh;
+				TempMesh->GetSourceModel(LODIndex).LoadRawMesh(RawMesh);
+
+				// Set RawMesh
+				OutStaticMesh->GetSourceModel(0).SaveRawMesh(RawMesh);
+			}
+		}
+			
+		// Copy Materials
+		const TArray<FStaticMaterial>& Materials = TempMesh->GetStaticMaterials();
+		OutStaticMesh->SetStaticMaterials(Materials);
+
+		// Done
+		TArray<FText> OutErrors;
+		OutStaticMesh->Build(true, &OutErrors);
+		OutStaticMesh->MarkPackageDirty();
+	}
+
+	// Destroy Temp Component and Actor
+	MeshComponent->UnregisterComponent();
+	MeshComponent->DestroyComponent();
+	Actor->Destroy();
+
+	return bGeneratedCorrectly ? OutStaticMesh : nullptr;
+}
+
+
+void UAnimToTextureBPLibrary::NormalizeVertexData(
+	const TArray<FVector3f>& Deltas, const TArray<FVector3f>& Normals,
+	FVector& MinBBox, FVector& SizeBBox,
+	TArray<FVector3f>& NormalizedDeltas, TArray<FVector3f>& NormalizedNormals)
+{
+	check(Deltas.Num() == Normals.Num());
+
+	// ---------------------------------------------------------------------------
+	// Compute Bounding Box
+	//
+	MinBBox = { TNumericLimits<float>::Max(), TNumericLimits<float>::Max(), TNumericLimits<float>::Max() };
+	FVector3f MaxBBox = { TNumericLimits<float>::Min(), TNumericLimits<float>::Min(), TNumericLimits<float>::Min() };
+	
+	for (const FVector3f& Delta: Deltas)
+	{
+		// Find Min/Max BoundingBox
+		MinBBox.X = FMath::Min(Delta.X, MinBBox.X);
+		MinBBox.Y = FMath::Min(Delta.Y, MinBBox.Y);
+		MinBBox.Z = FMath::Min(Delta.Z, MinBBox.Z);
+
+		MaxBBox.X = FMath::Max(Delta.X, MaxBBox.X);
+		MaxBBox.Y = FMath::Max(Delta.Y, MaxBBox.Y);
+		MaxBBox.Z = FMath::Max(Delta.Z, MaxBBox.Z);
+	}
+
+	SizeBBox = (FVector)MaxBBox - MinBBox;
+
+	// ---------------------------------------------------------------------------
+	// Normalize Vertex Position Deltas
+	// Basically we want all deltas to be between [0, 1]
+	
+	// Compute Normalization Factor per-axis.
+	const FVector NormFactor = {
+		1.0f / static_cast<float>(SizeBBox.X),
+		1.0f / static_cast<float>(SizeBBox.Y),
+		1.0f / static_cast<float>(SizeBBox.Z) };
+
+	NormalizedDeltas.SetNumUninitialized(Deltas.Num());
+	for (int32 Index = 0; Index < Deltas.Num(); ++Index)
+	{
+		NormalizedDeltas[Index] = (FVector3f)(((FVector)Deltas[Index] - MinBBox) * NormFactor);
+	}
+
+	// ---------------------------------------------------------------------------
+	// Normalize Vertex Normals
+	// And move them to [0, 1]
+	
+	NormalizedNormals.SetNumUninitialized(Normals.Num());
+	for (int32 Index = 0; Index < Normals.Num(); ++Index)
+	{
+		NormalizedNormals[Index] = (Normals[Index].GetSafeNormal() + FVector3f::OneVector) * 0.5f;
+	}
+
+}
+
+void UAnimToTextureBPLibrary::NormalizeBoneData(
+	const TArray<FVector3f>& Positions, const TArray<FVector4>& Rotations,
+	FVector& MinBBox, FVector& SizeBBox, 
+	TArray<FVector3f>& NormalizedPositions, TArray<FVector4>& NormalizedRotations)
+{
+	check(Positions.Num() == Rotations.Num());
+
+	// ---------------------------------------------------------------------------
+	// Compute Position Bounding Box
+	//
+	MinBBox = { TNumericLimits<float>::Max(), TNumericLimits<float>::Max(), TNumericLimits<float>::Max() };
+	FVector3f MaxBBox = { TNumericLimits<float>::Min(), TNumericLimits<float>::Min(), TNumericLimits<float>::Min() };
+
+	for (const FVector3f& Position : Positions)
+	{
+		// Find Min/Max BoundingBox
+		MinBBox.X = FMath::Min(Position.X, MinBBox.X);
+		MinBBox.Y = FMath::Min(Position.Y, MinBBox.Y);
+		MinBBox.Z = FMath::Min(Position.Z, MinBBox.Z);
+
+		MaxBBox.X = FMath::Max(Position.X, MaxBBox.X);
+		MaxBBox.Y = FMath::Max(Position.Y, MaxBBox.Y);
+		MaxBBox.Z = FMath::Max(Position.Z, MaxBBox.Z);
+	}
+
+	SizeBBox = (FVector)MaxBBox - MinBBox;
+
+	// ---------------------------------------------------------------------------
+	// Normalize Bone Position.
+	// Basically we want all positions to be between [0, 1]
+
+	// Compute Normalization Factor per-axis.
+	const FVector NormFactor = {
+		1.0f / static_cast<float>(SizeBBox.X),
+		1.0f / static_cast<float>(SizeBBox.Y),
+		1.0f / static_cast<float>(SizeBBox.Z) };
+
+	NormalizedPositions.SetNumUninitialized(Positions.Num());
+	for (int32 Index = 0; Index < Positions.Num(); ++Index)
+	{
+		NormalizedPositions[Index] = FVector3f(((FVector)Positions[Index] - MinBBox) * NormFactor);
+	}
+
+	// ---------------------------------------------------------------------------
+	// Normalize Rotations
+	// And move them to [0, 1]
+	NormalizedRotations.SetNumUninitialized(Rotations.Num());
+	for (int32 Index = 0; Index < Rotations.Num(); ++Index)
+	{
+		const FVector4 Axis = Rotations[Index];
+		const float Angle = Rotations[Index].W; // Angle are returned in radians and they go from [0-pi*2]
+
+		NormalizedRotations[Index] = (Axis.GetSafeNormal() + FVector::OneVector) * 0.5f;
+		NormalizedRotations[Index].W = Angle / (PI * 2.f);
+	}
+}
+
+
+bool UAnimToTextureBPLibrary::CreateUVChannel(
+	UStaticMesh* StaticMesh, const int32 LODIndex, const int32 UVChannelIndex,
+	const int32 Height, const int32 Width)
+{
+	if (!StaticMesh)
+	{
+		return false;
+	}
+
+	// Get Render Data
+	const FStaticMeshRenderData* RenderData = StaticMesh->GetRenderData();
+	check(RenderData);
+
+	// Get LOD Data
+	if (!RenderData->LODResources.IsValidIndex(LODIndex))
+	{
+		UE_LOG(LogTemp, Warning, TEXT("Invalid LODIndex: %i"), LODIndex);
+		return false;
+	}
+
+	// ----------------------------------------------------------------------------
+	// Get Mesh Description.
+	// This is needed for Inserting UVChannel
+	FMeshDescription* MeshDescription = StaticMesh->GetMeshDescription(LODIndex);
+	check(MeshDescription);
+
+	// Check if UVChannel is being used by the Lightmap UV
+	/*if (StaticMesh->GetLightMapCoordinateIndex() == UVChannelIndex)
+	{
+		UE_LOG(LogTemp, Warning, TEXT("Invalid UVChannel: %i. Already used by LightMap"), UVChannelIndex);
+		return false;
+	}*/
+
+	// Add New UVChannel.
+	// UE_LOG(LogTemp, Warning, TEXT("UVChannel: %i. Number of existing UVChannels: %i"), UVChannelIndex, StaticMesh->GetNumUVChannels(LODIndex));
+	if (UVChannelIndex == StaticMesh->GetNumUVChannels(LODIndex))
+	{
+		if (!StaticMesh->InsertUVChannel(LODIndex, UVChannelIndex))
+		{
+			UE_LOG(LogTemp, Warning, TEXT("Unable to Add UVChannel"));
+			return false;
+		}
+	}
+	else if (UVChannelIndex > StaticMesh->GetNumUVChannels(LODIndex))
+	{
+		UE_LOG(LogTemp, Warning, TEXT("UVChannel: %i Out of Range. Number of existing UVChannels: %i"), UVChannelIndex, StaticMesh->GetNumUVChannels(LODIndex));
+		return false;
+	}
+
+	// -----------------------------------------------------------------------------
+
+	TMap<FVertexInstanceID, FVector2D> TexCoords;
+
+	for (const FVertexInstanceID VertexInstanceID : MeshDescription->VertexInstances().GetElementIDs())
+	{
+		const FVertexID VertexID = MeshDescription->GetVertexInstanceVertex(VertexInstanceID);
+		const int32 VertexIndex = VertexID.GetValue();
+
+		float U = (0.5f / (float)Width) + (VertexIndex % Width) / (float)Width;
+		float V = (0.5f / (float)Height) + (VertexIndex / Width) / (float)Height;
+		
+		TexCoords.Add(VertexInstanceID, FVector2D(U, V));
+	}
+
+	// Set Full Precision UVs
+	SetFullPrecisionUVs(StaticMesh, true);
+
+	if (StaticMesh->SetUVChannel(LODIndex, UVChannelIndex, TexCoords))
+	{		
+		// Update and Mark to Save.
+		StaticMesh->MarkPackageDirty();
+
+		return true;
+	}
+	else
+	{
+		UE_LOG(LogTemp, Warning, TEXT("Unable to Set UVChannel: %i. TexCoords: %i"), UVChannelIndex, TexCoords.Num());
+		return false;
+	};
+
+	return false;
+}
+
+bool UAnimToTextureBPLibrary::FindBestResolution(
+	const int32 NumFrames, const int32 NumElements,
+	int32& OutHeight, int32& OutWidth, int32& OutRowsPerFrame,
+	const int32 MaxHeight, const int32 MaxWidth, bool bEnforcePowerOfTwo)
+{
+	if (bEnforcePowerOfTwo)
+	{
+		OutWidth = 2;
+		while (OutWidth < NumElements && OutWidth < MaxWidth)
+		{
+			OutWidth *= 2;
+		}
+		OutRowsPerFrame = FMath::CeilToInt(NumElements / (float)OutWidth);
+
+		const int32 TargetHeight = NumFrames * OutRowsPerFrame;
+		OutHeight = 2;
+		while (OutHeight < TargetHeight)
+		{
+			OutHeight *= 2;
+		}
+	}
+	else
+	{
+		OutRowsPerFrame = FMath::CeilToInt(NumElements / (float)MaxWidth);
+		OutWidth = FMath::CeilToInt(NumElements / (float)OutRowsPerFrame);
+		OutHeight = NumFrames * OutRowsPerFrame;
+	}
+
+	const bool bValidResolution = OutWidth <= MaxWidth && OutHeight <= MaxHeight;
+	return bValidResolution;
+};
+
+void UAnimToTextureBPLibrary::SetFullPrecisionUVs(UStaticMesh* StaticMesh, bool bFullPrecision)
+{
+	int32 NumSourceModels = StaticMesh->GetNumSourceModels();
+	for (int32 LodIndex = 0; LodIndex < NumSourceModels; LodIndex++)
+	{
+		FStaticMeshSourceModel& SourceModel = StaticMesh->GetSourceModel(LodIndex);
+		SourceModel.BuildSettings.bUseFullPrecisionUVs = bFullPrecision;
+	}
+}
+
+void UAnimToTextureBPLibrary::SetStaticMeshBoundsExtensions(
+	UStaticMesh* StaticMesh, 
+	const FVector& PositiveBoundsExtension, 
+	const FVector& NegativeBoundsExtension)
+{
+	if (IsValid(StaticMesh))
+	{
+		StaticMesh->SetPositiveBoundsExtension(PositiveBoundsExtension);
+		StaticMesh->SetNegativeBoundsExtension(NegativeBoundsExtension);
+	}
+}
diff --git a/Engine/Plugins/Animation/AnimToTexture/Source/AnimToTextureEditor/Private/AnimToTextureEditor.cpp b/Engine/Plugins/Animation/AnimToTexture/Source/AnimToTextureEditor/Private/AnimToTextureEditor.cpp
new file mode 100644
--- /dev/null
+++ b/Engine/Plugins/Animation/AnimToTexture/Source/AnimToTextureEditor/Private/AnimToTextureEditor.cpp
@@ -0,0 +1,22 @@
+// Copyright Epic Games, Inc. All Rights Reserved.
+
+#include "AnimToTextureEditor.h"
+
+#define LOCTEXT_NAMESPACE "FAnimToTextureEditorModule"
+
+void FAnimToTextureEditorModule::StartupModule()
+{
+	// This code will execute after your module is loaded into memory; the exact timing is specified in the .uplugin file per-module
+	
+}
+
+void FAnimToTextureEditorModule::ShutdownModule()
+{
+	// This function may be called during shutdown to clean up your module.  For modules that support dynamic reloading,
+	// we call this function before unloading the module.
+	
+}
+
+#undef LOCTEXT_NAMESPACE
+	
+IMPLEMENT_MODULE(FAnimToTextureEditorModule, AnimToTextureEditor)
\ No newline at end of file
diff --git a/Engine/Plugins/Animation/AnimToTexture/Source/AnimToTextureEditor/Private/AnimToTextureSkeletalMesh.cpp b/Engine/Plugins/Animation/AnimToTexture/Source/AnimToTextureEditor/Private/AnimToTextureSkeletalMesh.cpp
new file mode 100644
--- /dev/null
+++ b/Engine/Plugins/Animation/AnimToTexture/Source/AnimToTextureEditor/Private/AnimToTextureSkeletalMesh.cpp
@@ -0,0 +1,446 @@
+// Copyright Epic Games, Inc. All Rights Reserved.
+#include "AnimToTextureSkeletalMesh.h"
+#include "Rendering/SkeletalMeshModel.h"
+#include "Rendering/SkeletalMeshRenderData.h"
+#include "Engine/StaticMesh.h"
+#include "Engine/SkeletalMesh.h"
+#include "Components/SkeletalMeshComponent.h"
+#include "MeshDescription.h"
+#include "BoneIndices.h"
+
+namespace AnimToTexture_Private
+{
+
+int32 GetVertices(const USkeletalMesh& Mesh, const int32 LODIndex, 
+	TArray<FVector3f>& OutPositions, TArray<FVector3f>& OutNormals)
+{
+	OutPositions.Reset();
+	OutNormals.Reset();
+	
+	const FSkeletalMeshRenderData* RenderData = Mesh.GetResourceForRendering();
+	check(RenderData);
+
+	if (!RenderData->LODRenderData.IsValidIndex(LODIndex))
+	{
+		return INDEX_NONE;
+	}
+
+	// Get LOD Data
+	const FSkeletalMeshLODRenderData& LODRenderData = RenderData->LODRenderData[LODIndex];
+
+	// Get Total Num of Vertices (for all sections)
+	const int32 NumVertices = LODRenderData.GetNumVertices();
+	OutPositions.SetNumUninitialized(NumVertices);
+	OutNormals.SetNumUninitialized(NumVertices);
+	
+	for (int32 VertexIndex = 0; VertexIndex < NumVertices; ++VertexIndex)
+	{
+		OutPositions[VertexIndex] = LODRenderData.StaticVertexBuffers.PositionVertexBuffer.VertexPosition(VertexIndex);
+		OutNormals[VertexIndex] = LODRenderData.StaticVertexBuffers.StaticMeshVertexBuffer.VertexTangentZ(VertexIndex);
+	};
+
+	return NumVertices;
+}
+
+int32 FindClosestVertex(const FVector3f& Point, const TArray<FVector3f>& Vertices)
+{
+	float MinDistance = TNumericLimits<float>::Max();
+	int32 ClosestIndex = INDEX_NONE;
+
+	// Find Closest SkeletalMesh Vertex.
+	for (int32 VertexIndex = 0; VertexIndex < Vertices.Num(); ++VertexIndex)
+	{
+		const float Distance = FVector3f::Dist(Point, Vertices[VertexIndex]);
+		if (Distance < MinDistance)
+		{
+			MinDistance = Distance;
+			ClosestIndex = VertexIndex;
+		}
+	}
+
+	return ClosestIndex;
+}
+
+void GetStaticToSkeletalMapping(
+	const UStaticMesh& StaticMesh, const int32 StaticLODIndex,
+	const USkeletalMesh& SkeletalMesh, const int32 SkeletalLODIndex,
+	TArray<FVector3f>& OutStaticVertices, TArray<FVector3f>& OutSkeletalVertices,
+	TArray<int32>& OutStaticToSkeletalMapping)
+{
+	// Reset
+	OutStaticVertices.Reset();
+	OutSkeletalVertices.Reset();
+	OutStaticToSkeletalMapping.Reset();
+
+	// Get Mesh Description
+	const FMeshDescription* MeshDescription = StaticMesh.GetMeshDescription(StaticLODIndex);
+	check(MeshDescription);
+
+	// Get SkeletalMesh Vertices
+	TArray<FVector3f> Normals;
+	const int32 SkelNumVertices = AnimToTexture_Private::GetVertices(SkeletalMesh, SkeletalLODIndex, OutSkeletalVertices, Normals);
+
+	// Get MeshDescription Vertices
+	const FVertexArray& Vertices = MeshDescription->Vertices();
+	const int32 NumVertices = Vertices.Num();
+	OutStaticVertices.SetNum(NumVertices);
+	OutStaticToSkeletalMapping.SetNum(NumVertices);
+
+	for (const FVertexID VertexID : Vertices.GetElementIDs())
+	{
+		const int32 VertexIndex = VertexID.GetValue();
+
+		// Get Mesh Description Vertex Position
+		const FVector3f Vertex = MeshDescription->GetVertexPosition(VertexID);
+		OutStaticVertices[VertexIndex] = Vertex;
+
+		// Find Closest SkeletalMesh Vertex
+		OutStaticToSkeletalMapping[VertexIndex] = FindClosestVertex(Vertex, OutSkeletalVertices);
+	}
+}
+
+int32 GetNumBones(const USkeletalMesh* SkeletalMesh)
+{
+	check(SkeletalMesh);
+	
+	const FReferenceSkeleton& RefSkeleton = SkeletalMesh->GetRefSkeleton();
+	return RefSkeleton.GetRawBoneNum();
+}
+
+/**
+* Copied from FAnimationRuntime::GetComponentSpaceTransform
+*/
+void GetRefBoneTransforms(const USkeletalMesh* SkeletalMesh, TArray<FTransform>& Transforms)
+{
+	check(SkeletalMesh);
+	Transforms.Reset();
+
+	// Get Reference Skeleton
+	const FReferenceSkeleton& RefSkeleton = SkeletalMesh->GetRefSkeleton();
+	const TArray<FTransform>& BoneSpaceTransforms = RefSkeleton.GetRawRefBonePose(); // Get only raw bones (no virtual)
+	
+	const int32 NumTransforms = BoneSpaceTransforms.Num();
+	Transforms.SetNumUninitialized(NumTransforms);
+
+	for (int32 BoneIndex = 0; BoneIndex < NumTransforms; ++BoneIndex)
+	{
+		// initialize to identity since some of them don't have tracks
+		int32 IterBoneIndex = BoneIndex;
+		FTransform ComponentSpaceTransform = BoneSpaceTransforms[BoneIndex];
+
+		do
+		{
+			//const int32 ParentIndex = RefSkeleton.GetParentIndex(IterBoneIndex);
+			const int32 ParentIndex = RefSkeleton.GetRawParentIndex(IterBoneIndex); // Get only raw bones (no virtual)
+			if (ParentIndex != INDEX_NONE)
+			{
+				ComponentSpaceTransform *= BoneSpaceTransforms[ParentIndex];
+			}
+
+			IterBoneIndex = ParentIndex;
+		} while (RefSkeleton.IsValidIndex(IterBoneIndex));
+
+		// 
+		Transforms[BoneIndex] = ComponentSpaceTransform;
+	}
+}
+
+bool HasBone(const USkeletalMesh* SkeletalMesh, const FName& Bone)
+{
+	check(SkeletalMesh);
+	
+	const FReferenceSkeleton& RefSkeleton = SkeletalMesh->GetRefSkeleton();
+
+	for (const FMeshBoneInfo& BoneInfo: RefSkeleton.GetRawRefBoneInfo())
+	{
+		if (BoneInfo.Name == Bone)
+		{
+			return true;
+		}
+	}
+	return false;
+}
+
+
+void GetBoneNames(const USkeletalMesh* SkeletalMesh, TArray<FName>& OutNames)
+{
+	check(SkeletalMesh);
+	OutNames.Reset();
+	
+	const FReferenceSkeleton& RefSkeleton = SkeletalMesh->GetRefSkeleton();
+	
+	const int32 NumRawBones = RefSkeleton.GetRawBoneNum();
+	OutNames.SetNumUninitialized(NumRawBones);
+
+	for (int32 BoneIndex = 0; BoneIndex < NumRawBones; ++BoneIndex)
+	{
+		OutNames[BoneIndex] = RefSkeleton.GetBoneName(BoneIndex);
+	}
+};
+
+void GetSkinnedVertices(USkeletalMeshComponent* MeshComponent, const int32 LODIndex, 
+	TArray<FVector3f>& OutPositions, TArray<FVector3f>& OutNormals)
+{
+	OutPositions.Reset();
+	OutNormals.Reset();
+
+	if (!MeshComponent)
+	{
+		return;
+	}
+
+	// Get SkeletalMesh
+	USkeletalMesh* SkeletalMesh = MeshComponent->GetSkeletalMesh();
+	if (!SkeletalMesh)
+	{
+		return;
+	}
+
+	// Get Render Data 
+	FSkeletalMeshRenderData* RenderData = SkeletalMesh->GetResourceForRendering();
+	if (!RenderData->LODRenderData.IsValidIndex(LODIndex))
+	{
+		return;
+	};
+
+	// Get Matrices
+	TArray<FMatrix44f> RefToLocals;
+	MeshComponent->CacheRefToLocalMatrices(RefToLocals);
+
+	// ---------------------------------------------------------------------------------------------
+
+	// Get Ref-Pose Vertices
+	TArray<FVector3f> Vertices;
+	TArray<FVector3f> Normals;
+	const int32 NumVertices = GetVertices(*SkeletalMesh, LODIndex, Vertices, Normals);
+
+	// TODO: Add Morph Deltas to Vertices.
+
+	// Get Weights
+	TArray<TVertexSkinWeight<MAX_TOTAL_INFLUENCES>> SkinWeights;
+	GetSkinWeightsData(*SkeletalMesh, LODIndex, SkinWeights);
+
+	OutPositions.SetNumUninitialized(NumVertices);
+	OutNormals.SetNumUninitialized(NumVertices);
+
+	for (int32 VertexIndex = 0; VertexIndex < NumVertices; ++VertexIndex)
+	{
+		const FVector3f& Vertex = Vertices[VertexIndex];
+		const FVector3f& Normal = Normals[VertexIndex];
+		const TVertexSkinWeight<MAX_TOTAL_INFLUENCES>& Weights = SkinWeights[VertexIndex];
+
+		FVector4f SkinnedVertex(0);
+		FVector4f SkinnedNormal(0);
+
+		for (int32 InfluenceIndex = 0; InfluenceIndex < MAX_TOTAL_INFLUENCES; InfluenceIndex++)
+		{
+			const uint8& BoneWeight = Weights.BoneWeights[InfluenceIndex];
+			const uint16& MeshBoneIndex = Weights.MeshBoneIndices[InfluenceIndex];
+
+			// Get Matrix
+			const FMatrix44f& RefToLocal = RefToLocals[MeshBoneIndex];
+
+			const float Weight = (float)BoneWeight / 255.f;
+			SkinnedVertex += RefToLocal.TransformPosition(Vertex) * Weight;
+			SkinnedNormal += RefToLocal.TransformVector(Normal) * Weight;
+		}
+
+		OutPositions[VertexIndex] = SkinnedVertex;
+		OutNormals[VertexIndex] = SkinnedNormal;
+	};
+	
+};
+
+void ReduceSkinWeightsData(const TArray<TVertexSkinWeight<MAX_TOTAL_INFLUENCES>>& InSkinWeights, TArray<TVertexSkinWeight<4>>& OutSkinWeights)
+{
+	OutSkinWeights.SetNumUninitialized(InSkinWeights.Num());
+
+	for (int32 VertexIndex = 0; VertexIndex < InSkinWeights.Num(); ++VertexIndex)
+	{
+		GetMaxFourInfluences(InSkinWeights[VertexIndex], OutSkinWeights[VertexIndex]);
+	}
+}
+
+void GetFollowerBoneMap(const USkeletalMesh* LeaderSkeletalMesh, const USkeletalMesh* FollowerSkeletalMesh, TArray<int32>& FollowerBoneMap)
+{
+	const TArray<FMeshBoneInfo>& LeaderBoneInfo = LeaderSkeletalMesh->GetRefSkeleton().GetRawRefBoneInfo();
+	const TArray<FMeshBoneInfo>& FollowerBoneInfo = FollowerSkeletalMesh->GetRefSkeleton().GetRawRefBoneInfo();
+
+	// Initialize Array
+	const int32 LeaderNumBones = LeaderBoneInfo.Num();
+	const int32 FollowerNumBones = FollowerBoneInfo.Num();
+	FollowerBoneMap.Init(INDEX_NONE, FollowerNumBones);
+
+	for (int32 FollowerBoneIndex = 0; FollowerBoneIndex < FollowerNumBones; ++FollowerBoneIndex)
+	{
+		// Search by name-matching
+		for (int32 LeaderBoneIndex = 0; LeaderBoneIndex < LeaderNumBones; ++LeaderBoneIndex)
+		{
+			// found by name-matching
+			if (FollowerBoneInfo[FollowerBoneIndex].Name == LeaderBoneInfo[LeaderBoneIndex].Name)
+			{
+				FollowerBoneMap[FollowerBoneIndex] = LeaderBoneIndex;
+				break;
+			}
+		}
+
+		// Search in upstream hierarchy
+		if (FollowerBoneMap[FollowerBoneIndex] == INDEX_NONE)
+		{
+			bool FoundLeaderBone = false;
+			int32 ParentIndex = FollowerBoneInfo[FollowerBoneIndex].ParentIndex;
+			while (!FoundLeaderBone && ParentIndex != INDEX_NONE)
+			{
+				for (int32 LeaderBoneIndex = 0; LeaderBoneIndex < LeaderNumBones; ++LeaderBoneIndex)
+				{
+					if (FollowerBoneInfo[ParentIndex].Name == LeaderBoneInfo[LeaderBoneIndex].Name)
+					{
+						FollowerBoneMap[FollowerBoneIndex] = LeaderBoneIndex;
+						FoundLeaderBone = true;
+						break;
+					}
+				}
+
+				ParentIndex = FollowerBoneInfo[ParentIndex].ParentIndex;
+			}
+		}	
+
+		// FollowerBone not found in Leader
+		if (FollowerBoneMap[FollowerBoneIndex] == INDEX_NONE)
+		{
+			UE_LOG(LogTemp, Warning, TEXT("Unable to find BoneMapping for: %s"), *FollowerBoneInfo[FollowerBoneIndex].Name.ToString())
+		}
+	}
+}
+
+void DecomposeTransformation(const FTransform& Transform, FVector3f& Translation, FVector4& Rotation)
+{
+	// Get Translation
+	Translation = (FVector3f)Transform.GetTranslation();
+
+	// Get Rotation 
+	const FQuat Quat = Transform.GetRotation();
+
+	FVector Axis;
+	float Angle;
+	Quat.ToAxisAndAngle(Axis, Angle);
+
+	Rotation = FVector4(Axis, Angle);
+}
+
+void DecomposeTransformations(const TArray<FTransform>& Transforms, TArray<FVector3f>& Translations, TArray<FVector4>& Rotations)
+{
+	const int32 NumTransforms = Transforms.Num();
+	Translations.SetNumUninitialized(NumTransforms);
+	Rotations.SetNumUninitialized(NumTransforms);
+
+	for (int32 Index = 0; Index < NumTransforms; ++Index)
+	{
+		DecomposeTransformation(Transforms[Index], Translations[Index], Rotations[Index]);
+	}
+};
+
+
+void GetSkinWeightsData(const USkeletalMesh& SkeletalMesh, const int32 LODIndex, TArray<TVertexSkinWeight<MAX_TOTAL_INFLUENCES>>& SkinWeights)
+{
+	// Reset Weights
+	SkinWeights.Reset();
+
+	// Get Render Data
+	const FSkeletalMeshRenderData* RenderData = SkeletalMesh.GetResourceForRendering();
+	check(RenderData);
+
+	if (!RenderData->LODRenderData.IsValidIndex(LODIndex))
+	{
+		UE_LOG(LogTemp, Warning, TEXT("Invalid LODIndex: %i"), LODIndex)
+		return;
+	}
+
+	// Get LOD Data
+	const FSkeletalMeshLODRenderData& LODRenderData = RenderData->LODRenderData[LODIndex];
+
+	// Get Weights
+	const FSkinWeightVertexBuffer* SkinWeightVertexBuffer = LODRenderData.GetSkinWeightVertexBuffer();
+	check(SkinWeightVertexBuffer);
+
+	// Get Weights from Buffer.
+	// this is size of number of vertices.
+	TArray<FSkinWeightInfo> SkinWeightsInfo;
+	SkinWeightVertexBuffer->GetSkinWeights(SkinWeightsInfo);
+
+	// Allocated SkinWeightData
+	SkinWeights.SetNumUninitialized(SkinWeightsInfo.Num());
+
+	// Loop thru vertices
+	for (int32 VertexIndex = 0; VertexIndex < SkinWeightsInfo.Num(); ++VertexIndex)
+	{
+		// Find Section From Global Vertex Index
+		// NOTE: BoneMap is stored by Section.
+		int32 OutSectionIndex;
+		int32 OutSectionVertexIndex;
+		LODRenderData.GetSectionFromVertexIndex(VertexIndex, OutSectionIndex, OutSectionVertexIndex);
+
+		// Get Section for Vertex.
+		const FSkelMeshRenderSection& RenderSection = LODRenderData.RenderSections[OutSectionIndex];
+		
+		// Get Vertex Weights
+		const FSkinWeightInfo& SkinWeightInfo = SkinWeightsInfo[VertexIndex];
+
+		// Store Weights
+		for (int32 InfluenceIndex = 0; InfluenceIndex < MAX_TOTAL_INFLUENCES; ++InfluenceIndex)
+		{
+			const uint8& BoneWeight = SkinWeightInfo.InfluenceWeights[InfluenceIndex];
+			const uint16& BoneIndex = SkinWeightInfo.InfluenceBones[InfluenceIndex];
+			const uint16& MeshBoneIndex = RenderSection.BoneMap[BoneIndex];
+
+			SkinWeights[VertexIndex].BoneWeights[InfluenceIndex] = BoneWeight;
+			SkinWeights[VertexIndex].BoneIndices[InfluenceIndex] = BoneIndex;
+			SkinWeights[VertexIndex].MeshBoneIndices[InfluenceIndex] = MeshBoneIndex;
+		}
+	}
+}
+
+void GetReducedSkinWeightsData(
+	const USkeletalMesh& SkeletalMesh, const int32 SkeletalLODIndex, 
+	const TArray<int32> StaticToSkelMapping, 
+	TArray<TVertexSkinWeight<4>>& OutSkinWeights)
+{	
+	// Get SkeletalMesh Weights
+	TArray<TVertexSkinWeight<MAX_TOTAL_INFLUENCES>> SkinWeights;
+	AnimToTexture_Private::GetSkinWeightsData(SkeletalMesh, SkeletalLODIndex, SkinWeights);
+	
+	// Reduce SkeletalMesh Weights to 4 highest influences.
+	TArray<TVertexSkinWeight<4>> ReducedSkinWeights;
+	AnimToTexture_Private::ReduceSkinWeightsData(SkinWeights, ReducedSkinWeights);
+
+	// Allocate StaticMesh Weights
+	OutSkinWeights.SetNumUninitialized(StaticToSkelMapping.Num());
+
+	// Loop thru SkeletalMesh Weights
+	//uint16 InvalidMeshBoneIndex;
+	//bool bIsValidMeshBoneIndex = true;
+	for (int32 VertexIndex = 0; VertexIndex < StaticToSkelMapping.Num(); ++VertexIndex)
+	{
+		const int32 SkelVertexIndex = StaticToSkelMapping[VertexIndex];
+		OutSkinWeights[VertexIndex] = ReducedSkinWeights[SkelVertexIndex];
+
+		// MeshBoneIndex can not be higher than 255
+		/*for (uint16 MeshBoneIndex: OutSkinWeights[VertexIndex].MeshBoneIndices)
+		{
+			if (MeshBoneIndex > 255)
+			{
+				InvalidMeshBoneIndex = MeshBoneIndex;
+				bIsValidMeshBoneIndex = false;
+				break;
+			}
+		}
+
+		if (!bIsValidMeshBoneIndex)
+		{
+			UE_LOG(LogTemp, Warning, TEXT("Invalid SkinWeights on Vertex: %i. MeshBoneIndex: %i > 255"), VertexIndex, InvalidMeshBoneIndex)
+			break;
+		}*/
+	}
+}
+
+} // end namespace AnimToTexture_Private
\ No newline at end of file
diff --git a/Engine/Plugins/Animation/AnimToTexture/Source/AnimToTextureEditor/Private/AnimToTextureUtils.cpp b/Engine/Plugins/Animation/AnimToTexture/Source/AnimToTextureEditor/Private/AnimToTextureUtils.cpp
new file mode 100644
--- /dev/null
+++ b/Engine/Plugins/Animation/AnimToTexture/Source/AnimToTextureEditor/Private/AnimToTextureUtils.cpp
@@ -0,0 +1,83 @@
+// Copyright Epic Games, Inc. All Rights Reserved.
+
+#include "AnimToTextureUtils.h"
+#include "AnimToTextureSkeletalMesh.h"
+
+#include "Rendering/SkeletalMeshRenderData.h"
+
+namespace AnimToTexture_Private
+{
+	bool CheckSkinWeightsToTexture(const TArray<TVertexSkinWeight<4>>& SkinWeights)
+	{
+		bool bHasInvalidIndices = false;
+		int16 InvalidMeshBoneIndex = INDEX_NONE;
+		for (int32 VertexIndex = 0; VertexIndex < SkinWeights.Num(); ++VertexIndex)
+		{
+			// MeshBoneIndex can not be higher than 255
+			for (uint16 MeshBoneIndex : SkinWeights[VertexIndex].MeshBoneIndices)
+			{
+				if (MeshBoneIndex > 255)
+				{
+					InvalidMeshBoneIndex = MeshBoneIndex;
+					bHasInvalidIndices = true;
+					break;
+				}
+			}
+
+			if (InvalidMeshBoneIndex != INDEX_NONE)
+			{
+				UE_LOG(LogTemp, Warning, TEXT("Invalid SkinWeights on Vertex: %i. MeshBoneIndex: %i > 255"), VertexIndex, InvalidMeshBoneIndex)
+			}
+		}
+
+		if (bHasInvalidIndices)
+		{
+			return false;
+		}
+
+		return true;
+	}
+
+	bool WriteSkinWeightsToTexture(const TArray<TVertexSkinWeight<4>>& SkinWeights,
+		const int32 RowsPerFrame, const int32 Height, const int32 Width, UTexture2D* Texture)
+	{
+		if (!Texture)
+		{
+			return false;
+		}
+
+		// Sanity check for MeshBoneIndices > 256
+		if (!CheckSkinWeightsToTexture(SkinWeights))
+		{
+			return false;
+		}
+
+		const int32 NumVertices = SkinWeights.Num();
+
+		// Allocate PixelData.
+		TArray<FColor> Pixels;
+		Pixels.Init(FColor::Black, Height * Width);
+
+		for (int32 VertexIndex = 0; VertexIndex < NumVertices; ++VertexIndex)
+		{
+			const TVertexSkinWeight<4>& VertexSkinWeight = SkinWeights[VertexIndex];
+
+			// Write Influence
+			// NOTE: we are assuming the bone index is under < 256
+			Pixels[VertexIndex].R = (uint8)VertexSkinWeight.MeshBoneIndices[0];
+			Pixels[VertexIndex].G = (uint8)VertexSkinWeight.MeshBoneIndices[1];
+			Pixels[VertexIndex].B = (uint8)VertexSkinWeight.MeshBoneIndices[2];
+			Pixels[VertexIndex].A = (uint8)VertexSkinWeight.MeshBoneIndices[3];
+			
+			// Write Weight
+			Pixels[Width * RowsPerFrame + VertexIndex].R = VertexSkinWeight.BoneWeights[0];
+			Pixels[Width * RowsPerFrame + VertexIndex].G = VertexSkinWeight.BoneWeights[1];
+			Pixels[Width * RowsPerFrame + VertexIndex].B = VertexSkinWeight.BoneWeights[2];
+			Pixels[Width * RowsPerFrame + VertexIndex].A = VertexSkinWeight.BoneWeights[3];
+		};
+
+		// Write to Texture
+		return WriteToTexture<FLowPrecision>(Texture, Height, Width, Pixels);
+	}
+
+}
\ No newline at end of file
diff --git a/Engine/Plugins/Animation/AnimToTexture/Source/AnimToTextureEditor/Public/AnimToTextureBPLibrary.h b/Engine/Plugins/Animation/AnimToTexture/Source/AnimToTextureEditor/Public/AnimToTextureBPLibrary.h
new file mode 100644
--- /dev/null
+++ b/Engine/Plugins/Animation/AnimToTexture/Source/AnimToTextureEditor/Public/AnimToTextureBPLibrary.h
@@ -0,0 +1,86 @@
+// Copyright Epic Games, Inc. All Rights Reserved.
+
+#pragma once
+
+#include "AnimToTextureDataAsset.h"
+#include "Kismet/BlueprintFunctionLibrary.h"
+#include "Materials/MaterialLayersFunctions.h"
+
+#include "AnimToTextureBPLibrary.generated.h"
+
+class UStaticMesh;
+class USkeletalMesh;
+
+/*
+* TODO: 
+*	- Right now it is saving data per-vertex instead of per-point.
+*     This will require more pixels if the mesh has lots of material slots or uv-shells.
+*     FStaticMeshOperations::FindOverlappingCorners(FOverlappingCorners& OutOverlappingCorners, const FMeshDescription& MeshDescription, float ComparisonThreshold)
+*/
+UCLASS()
+class UAnimToTextureBPLibrary : public UBlueprintFunctionLibrary
+{
+	GENERATED_UCLASS_BODY()
+
+#if WITH_EDITOR
+
+	/**
+	* Bakes Animation Data into Textures.
+	* Position & Normals.
+	*/
+	UFUNCTION(BlueprintCallable, meta = (Category = "Vertex Animation"))
+	static void AnimationToTexture(UAnimToTextureDataAsset* DataAsset,
+		const FTransform RootTransform, const bool bCreateTextures=true, const bool bCreateUVChannel=true);
+
+	/* Utility for converting SkeletalMesh into a StaticMesh */
+	UFUNCTION(BlueprintCallable, Category = "Vertex Animation")
+	static UStaticMesh* ConvertSkeletalMeshToStaticMesh(USkeletalMesh* SkeletalMesh, const FString PackageName, const int32 LODIndex = -1);
+
+	/**
+	* Updates a material's parameters to match those of an animToTexture data asset
+	*/
+	UFUNCTION(BlueprintCallable, meta = (Category = "Vertex Animation"))
+	static void UpdateMaterialInstanceFromDataAsset(UAnimToTextureDataAsset* DataAsset, class UMaterialInstanceConstant* MaterialInstance, 
+		const FAnimToTextureMaterialParamNames& ParamNames, const EMaterialParameterAssociation MaterialParameterAssociation = EMaterialParameterAssociation::GlobalParameter);
+
+	/**
+	 * Replaces material layer function
+	 */
+	UFUNCTION(BlueprintCallable, meta = (Category = "Vertex Animation"))
+	static void UpdateMaterialLayerFunction(
+		class UMaterialInstanceConstant* MaterialInstance, 
+		class UMaterialFunctionInterface* OldMaterialFunction, 
+		class UMaterialFunctionInterface* NewMaterialFunction);
+
+	UFUNCTION(BlueprintCallable, meta = (Category = "Vertex Animation"))
+	static void SetStaticMeshBoundsExtensions(UStaticMesh* StaticMesh, const FVector& PositiveBoundsExtension, const FVector& NegativeBoundsExtension);
+
+#endif // WITH_EDITOR
+
+private:
+
+	static void NormalizeVertexData(
+		const TArray<FVector3f>& Deltas, const TArray<FVector3f>& Normals,
+		FVector& MinBBox, FVector& SizeBBox,
+		TArray<FVector3f>& NormalizedDeltas, TArray<FVector3f>& NormalizedNormals);
+
+	static void NormalizeBoneData(
+		const TArray<FVector3f>& Positions, const TArray<FVector4>& Rotations,
+		FVector& MinBBox, FVector& SizeBBox,
+		TArray<FVector3f>& NormalizedPositions, TArray<FVector4>& NormalizedRotations);
+
+	/* Returns best resolution for the given data. 
+	*  Returns false if data wont fit in the the max range.
+	*/
+	static bool FindBestResolution(const int32 NumFrames, const int32 NumElements,
+								   int32& OutHeight, int32& OutWidth, int32& OutRowsPerFrame,
+								   const int32 MaxHeight = 4096, const int32 MaxWidth = 4096, bool bEnforcePowerOfTwo = false);
+
+	/* Sets Static Mesh FullPrecisionUVs Property*/
+	static void SetFullPrecisionUVs(UStaticMesh* StaticMesh, bool bFullPrecision);
+
+	/* Creates UV Coord with vertices */
+	static bool CreateUVChannel(UStaticMesh* StaticMesh, const int32 LODIndex, const int32 UVChannelIndex,
+		const int32 Height, const int32 Width);
+
+};
diff --git a/Engine/Plugins/Animation/AnimToTexture/Source/AnimToTextureEditor/Public/AnimToTextureEditor.h b/Engine/Plugins/Animation/AnimToTexture/Source/AnimToTextureEditor/Public/AnimToTextureEditor.h
new file mode 100644
--- /dev/null
+++ b/Engine/Plugins/Animation/AnimToTexture/Source/AnimToTextureEditor/Public/AnimToTextureEditor.h
@@ -0,0 +1,14 @@
+// Copyright Epic Games, Inc. All Rights Reserved.
+
+#pragma once
+
+#include "Modules/ModuleManager.h"
+
+class FAnimToTextureEditorModule : public IModuleInterface
+{
+public:
+
+	/** IModuleInterface implementation */
+	virtual void StartupModule() override;
+	virtual void ShutdownModule() override;
+};
diff --git a/Engine/Plugins/Animation/AnimToTexture/Source/AnimToTextureEditor/Public/AnimToTextureSkeletalMesh.h b/Engine/Plugins/Animation/AnimToTexture/Source/AnimToTextureEditor/Public/AnimToTextureSkeletalMesh.h
new file mode 100644
--- /dev/null
+++ b/Engine/Plugins/Animation/AnimToTexture/Source/AnimToTextureEditor/Public/AnimToTextureSkeletalMesh.h
@@ -0,0 +1,224 @@
+// Copyright Epic Games, Inc. All Rights Reserved.
+
+#pragma once
+
+#include "CoreMinimal.h"
+#include "GPUSkinPublicDefs.h"
+#include "Containers/StaticArray.h"
+
+#include "AnimToTextureDataAsset.h"
+
+namespace AnimToTexture_Private
+{
+
+template <uint16 NumInfluences>
+struct TVertexSkinWeight
+{
+	TStaticArray<uint16, NumInfluences> MeshBoneIndices;
+	TStaticArray<uint16, NumInfluences> BoneIndices;
+	TStaticArray<uint8, NumInfluences>  BoneWeights;
+};
+
+/* Returns RefPose Vertex Positions */
+int32 GetVertices(const USkeletalMesh& Mesh, const int32 LODIndex, 
+	TArray<FVector3f>& OutPositions, TArray<FVector3f>& OutNormals);
+
+/* Finds closest vertex to point */
+int32 FindClosestVertex(const FVector3f& Point, const TArray<FVector3f>& Vertices);
+
+/* VertexMapping between StaticMesh and SkeletalMesh Vertices.
+*  Note: Size of Mapping is Number of MeshDescription Vertices.
+*/
+void GetStaticToSkeletalMapping(
+	const UStaticMesh& StaticMesh, const int32 StaticLODIndex,
+	const USkeletalMesh& SkeletalMesh, const int32 SkeletalLODIndex,
+	TArray<FVector3f>& OutStaticVertices, TArray<FVector3f>& OutSkeletalVertices,
+	TArray<int32>& OutStaticToSkeletalMapping);
+
+/* Computes CPUSkinning at Pose */
+void GetSkinnedVertices(USkeletalMeshComponent* MeshComponent, const int32 LODIndex,
+	TArray<FVector3f>& OutPositions, TArray<FVector3f>& OutNormals);
+
+/** Gets Skin Weights Data from SkeletalMeshComponent */
+void GetSkinWeightsData(const USkeletalMesh& SkeletalMesh, const int32 LODIndex, TArray<TVertexSkinWeight<MAX_TOTAL_INFLUENCES>>& SkinWeights);
+
+/** Gets Skin Weights Data from SkeletalMeshComponent, Reduce it to NumInfluences and Maps it to the StaticMesh */
+void GetReducedSkinWeightsData(
+	const USkeletalMesh& SkeletalMesh, const int32 SkeletalLODIndex, 
+	const TArray<int32> StaticToSkelMapping,
+	TArray<TVertexSkinWeight<4>>& SkinWeights);
+
+/** Reduce Weights */
+void ReduceSkinWeightsData(const TArray<TVertexSkinWeight<MAX_TOTAL_INFLUENCES>>& InSkinWeights, TArray<TVertexSkinWeight<4>>& OutSkinWeights);
+
+/* Returns the 4-Higher Influences and Weights */
+void GetMaxFourInfluences(const TVertexSkinWeight<MAX_TOTAL_INFLUENCES>& InSkinWeights, TVertexSkinWeight<4>& OutSkinWeights);
+
+/* Returns Number of RawBones (no virtual bones)*/
+int32 GetNumBones(const USkeletalMesh* SkeletalMesh);
+
+/* Returns RefPose Bone Transforms.
+   Only the RawBones are returned (no virtual bones)
+   The returned transforms are in ComponentSpace */
+void GetRefBoneTransforms(const USkeletalMesh* SkeletalMesh, TArray<FTransform>& OutTransforms);
+
+/* Returns Bone exist in the RawBone list (no virtual bones) */
+bool HasBone(const USkeletalMesh* SkeletalMesh, const FName& Bone);
+
+/* Returns Bone Names. 
+   Only the RawBones are returned (no virtual bones) */
+void GetBoneNames(const USkeletalMesh* SkeletalMesh, TArray<FName>& OutNames);
+
+/** Returns a Mapping between Follower and Leader. The hierarchy is traversed upstream until a matching Bone is found. */
+void GetFollowerBoneMap(const USkeletalMesh* LeaderSkeletalMesh, const USkeletalMesh* FollowerSkeletalMesh, TArray<int32>& FollowerBoneMap);
+
+/* Decomposes Transform in Translation and AxisAndAngle */
+void DecomposeTransformation(const FTransform& Transform, FVector3f& Translation, FVector4& Rotation);
+void DecomposeTransformations(const TArray<FTransform>& Transforms, TArray<FVector3f>& Translations, TArray<FVector4>& Rotations);
+
+// ----------------------------------------------------------------------------
+
+//template<>
+//FORCEINLINE void GetMaxInfluences<1>(
+//	const TVertexSkinWeight<MAX_TOTAL_INFLUENCES>& InSkinWeights, TVertexSkinWeight<4>& OutSkinWeights)
+//{
+//	// Reset Values
+//	for (int32 Index = 0; Index < 4; ++Index)
+//	{
+//		OutSkinWeights.BoneWeights[Index] = 0;
+//		OutSkinWeights.BoneIndices[Index] = 0;
+//		OutSkinWeights.MeshBoneIndices[Index] = 0;
+//	}
+//
+//	// Set first influence to 1
+//	OutSkinWeights.BoneWeights[0] = 255; // TNumericLimits<uint8>::Max();
+//
+//	// Find 1st Influence (Highest)
+//	uint8 MaxWeight = TNumericLimits<uint8>::Min();
+//	for (int32 Index = 0; Index < MAX_TOTAL_INFLUENCES; ++Index)
+//	{
+//		if (InSkinWeights.BoneWeights[Index] > MaxWeight)
+//		{
+//			MaxWeight = InSkinWeights.BoneWeights[Index];
+//			OutSkinWeights.BoneIndices[0] = InSkinWeights.BoneIndices[Index];
+//			OutSkinWeights.MeshBoneIndices[0] = InSkinWeights.MeshBoneIndices[Index];
+//		}
+//	}
+//}
+//
+//template<>
+//FORCEINLINE void GetMaxInfluences<2>(
+//	const TVertexSkinWeight<MAX_TOTAL_INFLUENCES>& InSkinWeights, TVertexSkinWeight<4>& OutSkinWeights)
+//{
+//	// Reset Values
+//	for (int32 Index = 0; Index < 4; ++Index)
+//	{
+//		OutSkinWeights.BoneIndices[Index] = 0;
+//		OutSkinWeights.BoneWeights[Index] = 0; 
+//		OutSkinWeights.MeshBoneIndices[Index] = 0;
+//	}
+//	
+//	// Find 1st Influence (Highest)
+//	for (int32 Index = 0; Index < MAX_TOTAL_INFLUENCES; ++Index)
+//	{
+//		if (InSkinWeights.BoneWeights[Index] > OutSkinWeights.BoneWeights[0])
+//		{
+//			OutSkinWeights.BoneWeights[0] = InSkinWeights.BoneWeights[Index];
+//			OutSkinWeights.BoneIndices[0] = InSkinWeights.BoneIndices[Index];
+//			OutSkinWeights.MeshBoneIndices[0] = InSkinWeights.MeshBoneIndices[Index];
+//		}
+//	}
+//
+//	// Find 2nd Influence
+//	for (int32 Index = 0; Index < MAX_TOTAL_INFLUENCES; ++Index)
+//	{
+//		if (InSkinWeights.BoneIndices[Index] != OutSkinWeights.BoneIndices[0])
+//		{
+//			if (InSkinWeights.BoneWeights[Index] > OutSkinWeights.BoneWeights[1])
+//			{
+//				OutSkinWeights.BoneWeights[1] = InSkinWeights.BoneWeights[Index];
+//				OutSkinWeights.BoneIndices[1] = InSkinWeights.BoneIndices[Index];
+//				OutSkinWeights.MeshBoneIndices[1] = InSkinWeights.MeshBoneIndices[Index];
+//			}
+//		}
+//	}
+//
+//	// Normalize Weighs
+//	const float NormFactor = 255.0f / (OutSkinWeights.BoneWeights[0] + OutSkinWeights.BoneWeights[1]);
+//	OutSkinWeights.BoneWeights[0] = FMath::RoundToInt(OutSkinWeights.BoneWeights[0] * NormFactor);
+//	OutSkinWeights.BoneWeights[1] = FMath::RoundToInt(OutSkinWeights.BoneWeights[1] * NormFactor);
+//}
+
+FORCEINLINE void GetMaxFourInfluences(
+	const TVertexSkinWeight<MAX_TOTAL_INFLUENCES>& InSkinWeights, TVertexSkinWeight<4>& OutSkinWeights)
+{
+	// Reset Values
+	for (int32 Index = 0; Index < 4; ++Index)
+	{
+		OutSkinWeights.BoneWeights[Index] = 0;
+		OutSkinWeights.BoneIndices[Index] = 0;
+		OutSkinWeights.MeshBoneIndices[Index] = 0;
+	}
+
+	// Find 1st Influence (Highest)
+	for (int32 Index = 0; Index < MAX_TOTAL_INFLUENCES; ++Index)
+	{
+		if (InSkinWeights.BoneWeights[Index] > OutSkinWeights.BoneWeights[0])
+		{
+			OutSkinWeights.BoneWeights[0] = InSkinWeights.BoneWeights[Index];
+			OutSkinWeights.BoneIndices[0] = InSkinWeights.BoneIndices[Index];
+			OutSkinWeights.MeshBoneIndices[0] = InSkinWeights.MeshBoneIndices[Index];
+		}
+	}
+
+	// Find 2nd Influence
+	for (int32 Index = 0; Index < MAX_TOTAL_INFLUENCES; ++Index)
+	{		
+		if (InSkinWeights.BoneIndices[Index] != OutSkinWeights.BoneIndices[0])
+		{
+			if (InSkinWeights.BoneWeights[Index] > OutSkinWeights.BoneWeights[1])
+			{
+				OutSkinWeights.BoneWeights[1] = InSkinWeights.BoneWeights[Index];
+				OutSkinWeights.BoneIndices[1] = InSkinWeights.BoneIndices[Index];
+				OutSkinWeights.MeshBoneIndices[1] = InSkinWeights.MeshBoneIndices[Index];
+			}
+		}
+	}
+
+	// Find 3rd Influence
+	for (int32 Index = 0; Index < MAX_TOTAL_INFLUENCES; ++Index)
+	{
+		if (InSkinWeights.BoneIndices[Index] != OutSkinWeights.BoneIndices[0] && InSkinWeights.BoneIndices[Index] != OutSkinWeights.BoneIndices[1])
+		{
+			if (InSkinWeights.BoneWeights[Index] > OutSkinWeights.BoneWeights[2])
+			{
+				OutSkinWeights.BoneWeights[2] = InSkinWeights.BoneWeights[Index];
+				OutSkinWeights.BoneIndices[2] = InSkinWeights.BoneIndices[Index];
+				OutSkinWeights.MeshBoneIndices[2] = InSkinWeights.MeshBoneIndices[Index];
+			}
+		}
+	}
+
+	// Find 4th Influence
+	for (int32 Index = 0; Index < MAX_TOTAL_INFLUENCES; ++Index)
+	{
+		if (InSkinWeights.BoneIndices[Index] != OutSkinWeights.BoneIndices[0] && InSkinWeights.BoneIndices[Index] != OutSkinWeights.BoneIndices[1] && InSkinWeights.BoneIndices[Index] != OutSkinWeights.BoneIndices[2])
+		{
+			if (InSkinWeights.BoneWeights[Index] > OutSkinWeights.BoneWeights[3])
+			{
+				OutSkinWeights.BoneWeights[3] = InSkinWeights.BoneWeights[Index];
+				OutSkinWeights.BoneIndices[3] = InSkinWeights.BoneIndices[Index];
+				OutSkinWeights.MeshBoneIndices[3] = InSkinWeights.MeshBoneIndices[Index];
+			}
+		}
+	}
+
+	// Normalize Weighs
+	const float NormFactor = 255.0f / (OutSkinWeights.BoneWeights[0] + OutSkinWeights.BoneWeights[1] + OutSkinWeights.BoneWeights[2] + OutSkinWeights.BoneWeights[3]);
+	for (int32 Index = 0; Index < 4; ++Index)
+	{
+		OutSkinWeights.BoneWeights[Index] = FMath::RoundToInt(OutSkinWeights.BoneWeights[Index] * NormFactor);
+	}
+}
+
+} // end namespace AnimToTexture_Private
diff --git a/Engine/Plugins/Animation/AnimToTexture/Source/AnimToTextureEditor/Public/AnimToTextureUtils.h b/Engine/Plugins/Animation/AnimToTexture/Source/AnimToTextureEditor/Public/AnimToTextureUtils.h
new file mode 100644
--- /dev/null
+++ b/Engine/Plugins/Animation/AnimToTexture/Source/AnimToTextureEditor/Public/AnimToTextureUtils.h
@@ -0,0 +1,309 @@
+// Copyright Epic Games, Inc. All Rights Reserved.
+
+#pragma once
+
+#include "CoreMinimal.h"
+#include "Engine/Texture.h"
+#include "Engine/Texture2D.h"
+
+namespace AnimToTexture_Private
+{
+	// Forward declaration
+	template <uint16 NumInfluences> struct TVertexSkinWeight;
+
+	struct FVector4u16
+	{
+		uint16 X;
+		uint16 Y;
+		uint16 Z;
+		uint16 W;
+	};
+
+	struct FLowPrecision
+	{
+		using ColorType = FColor;
+		static constexpr EPixelFormat PixelFormat = EPixelFormat::PF_B8G8R8A8;
+		static constexpr ETextureSourceFormat TextureSourceFormat = ETextureSourceFormat::TSF_BGRA8;
+		static constexpr TextureCompressionSettings CompressionSettings = TextureCompressionSettings::TC_VectorDisplacementmap;
+		static constexpr ColorType DefaultColor = FColor(0, 0, 0);
+	};
+
+	struct FHighPrecision
+	{
+		using ColorType = FVector4u16;
+		static constexpr EPixelFormat PixelFormat = EPixelFormat::PF_R16G16B16A16_UNORM;
+		static constexpr ETextureSourceFormat TextureSourceFormat = ETextureSourceFormat::TSF_RGBA16;
+		static constexpr TextureCompressionSettings CompressionSettings = TextureCompressionSettings::TC_HDR;
+		static constexpr ColorType DefaultColor = { 0, 0, 0, 0 };
+	};
+
+
+	/** Writes list of vectors into texture
+	*   Note: They must be pre-normalized. */
+	template<class V, class TextureSettings>
+	bool WriteVectorsToTexture(const TArray<V>& Vectors,
+		const int32 NumFrames, const int32 RowsPerFrame,
+		const int32 Height, const int32 Width, 
+		UTexture2D* Texture);
+
+	/* Writes list of skinweights into texture.
+	*  The SkinWeights data is already in uint8 & uint16 format, no need for normalizing it.
+	*/
+	bool WriteSkinWeightsToTexture(const TArray<TVertexSkinWeight<4>>& SkinWeights,
+		const int32 RowsPerFrame,
+		const int32 Height, const int32 Width,
+		UTexture2D* Texture);
+
+	/* Helper utility for writing 8 or 16 bits textures */
+	template<class TextureSettings>
+	bool WriteToTexture(UTexture2D* Texture, const uint32 Height, const uint32 Width, const TArray<typename TextureSettings::ColorType>& Data);
+
+	template<class V /* FVector / FVector4 / FIntVector */, class C /* FColor / FLinearColor */>
+	void VectorToColor(const V& Vector, C& Color);
+	
+	//template<class C /* FColor / FLinearColor */, class T /* uint8 / float */>
+	//void ColorArrayToData(const TArray<C>& Colors, TArray<T>& Data);
+}
+
+// ----------------------------------------------------------------------------
+
+template<>
+FORCEINLINE void AnimToTexture_Private::VectorToColor(const FVector3f& Vector, FColor& Color)
+{
+	const float ClampedX = FMath::Clamp(Vector.X, 0.0f, 1.0f);
+	const float ClampedY = FMath::Clamp(Vector.Y, 0.0f, 1.0f);
+	const float ClampedZ = FMath::Clamp(Vector.Z, 0.0f, 1.0f);
+
+	Color.R = (uint8)FMath::RoundToInt(ClampedX * 255.0f);
+	Color.G = (uint8)FMath::RoundToInt(ClampedY * 255.0f);
+	Color.B = (uint8)FMath::RoundToInt(ClampedZ * 255.0f);
+	Color.A = 255;
+}
+
+template<>
+FORCEINLINE void AnimToTexture_Private::VectorToColor(const FVector3d& Vector, FColor& Color)
+{
+	VectorToColor((FVector3f)Vector, Color);
+}
+
+
+template<>
+FORCEINLINE void AnimToTexture_Private::VectorToColor(const FVector4& Vector, FColor& Color)
+{
+	const float ClampedX = FMath::Clamp(Vector.X, 0.0f, 1.0f);
+	const float ClampedY = FMath::Clamp(Vector.Y, 0.0f, 1.0f);
+	const float ClampedZ = FMath::Clamp(Vector.Z, 0.0f, 1.0f);
+	const float ClampedW = FMath::Clamp(Vector.W, 0.0f, 1.0f);
+	
+	Color.R = (uint8)FMath::RoundToInt(ClampedX * 255.0f);
+	Color.G = (uint8)FMath::RoundToInt(ClampedY * 255.0f);
+	Color.B = (uint8)FMath::RoundToInt(ClampedZ * 255.0f);
+	Color.A = (uint8)FMath::RoundToInt(ClampedW * 255.0f);
+}
+
+template<>
+FORCEINLINE void AnimToTexture_Private::VectorToColor(const FIntVector& Vector, FColor& Color)
+{
+	const int32 ClampedX = FMath::Clamp(Vector.X, 0, 1);
+	const int32 ClampedY = FMath::Clamp(Vector.Y, 0, 1);
+	const int32 ClampedZ = FMath::Clamp(Vector.Z, 0, 1);
+
+	Color.R = (uint8)(ClampedX * 255);
+	Color.G = (uint8)(ClampedY * 255);
+	Color.B = (uint8)(ClampedZ * 255);
+	Color.A = 255;
+}
+
+template<>
+FORCEINLINE void AnimToTexture_Private::VectorToColor(const FIntVector4& Vector, FColor& Color)
+{
+	const int32 ClampedX = FMath::Clamp(Vector.X, 0, 1);
+	const int32 ClampedY = FMath::Clamp(Vector.Y, 0, 1);
+	const int32 ClampedZ = FMath::Clamp(Vector.Z, 0, 1);
+	const int32 ClampedW = FMath::Clamp(Vector.W, 0, 1);
+
+	Color.R = (uint8)(ClampedX * 255);
+	Color.G = (uint8)(ClampedY * 255);
+	Color.B = (uint8)(ClampedZ * 255);
+	Color.A = (uint8)(ClampedW * 255);
+}
+
+template<>
+FORCEINLINE void AnimToTexture_Private::VectorToColor(const FVector3f& Vector, FLinearColor& Color)
+{
+	Color.R = Vector.X;
+	Color.G = Vector.Y;
+	Color.B = Vector.Z;
+	Color.A = 1.0f;
+}
+
+template<>
+FORCEINLINE void AnimToTexture_Private::VectorToColor(const FVector4& Vector, FLinearColor& Color)
+{
+	Color.R = Vector.X;
+	Color.G = Vector.Y;
+	Color.B = Vector.Z;
+	Color.A = Vector.W;
+}
+
+template<>
+FORCEINLINE void AnimToTexture_Private::VectorToColor(const FVector3f& Vector, FVector4u16& Color)
+{
+	Color.X = FMath::RoundToInt(FMath::Clamp(Vector.X, 0.0f, 1.0f) * MAX_uint16);
+	Color.Y = FMath::RoundToInt(FMath::Clamp(Vector.Y, 0.0f, 1.0f) * MAX_uint16);
+	Color.Z = FMath::RoundToInt(FMath::Clamp(Vector.Z, 0.0f, 1.0f) * MAX_uint16);
+	Color.W = MAX_uint16;
+}
+
+template<>
+FORCEINLINE void AnimToTexture_Private::VectorToColor(const FVector4& Vector, FVector4u16& Color)
+{
+	Color.X = FMath::RoundToInt(FMath::Clamp(Vector.X, 0.0f, 1.0f) * MAX_uint16);
+	Color.Y = FMath::RoundToInt(FMath::Clamp(Vector.Y, 0.0f, 1.0f) * MAX_uint16);
+	Color.Z = FMath::RoundToInt(FMath::Clamp(Vector.Z, 0.0f, 1.0f) * MAX_uint16);
+	Color.W = FMath::RoundToInt(FMath::Clamp(Vector.W, 0.0f, 1.0f) * MAX_uint16);
+}
+
+/*
+template<>
+FORCEINLINE void AnimToTexture_Private::ColorArrayToData(const TArray<FColor>& Colors, TArray<uint8>& Data)
+{
+	// Resize Array
+	Data.Init(0, Colors.Num() * 4);
+
+	for (int32 Index = 0; Index < Colors.Num(); ++Index)
+	{
+		// NOTE: format is BGRA (swaping R <-> B)
+		Data[Index * 4 + 0] = Colors[Index].B; // B
+		Data[Index * 4 + 1] = Colors[Index].G; // G
+		Data[Index * 4 + 2] = Colors[Index].R; // R
+		Data[Index * 4 + 3] = Colors[Index].A; // A
+	}
+}
+
+template<>
+FORCEINLINE void AnimToTexture_Private::ColorArrayToData(const TArray<FLinearColor>& Colors, TArray<float>& Data)
+{
+	// Resize Array
+	Data.Init(0.0f, Colors.Num() * 4);
+
+	for (int32 Index = 0; Index < Colors.Num(); ++Index)
+	{
+		// NOTE: format is RGBA
+		Data[Index * 4 + 0] = Colors[Index].R; // R
+		Data[Index * 4 + 1] = Colors[Index].G; // G
+		Data[Index * 4 + 2] = Colors[Index].B; // B
+		Data[Index * 4 + 3] = Colors[Index].A; // A
+	}
+}
+*/
+
+
+template<class V, class TextureSettings>
+FORCEINLINE bool AnimToTexture_Private::WriteVectorsToTexture(const TArray<V>& Vectors,
+	const int32 NumFrames, const int32 RowsPerFrame,
+	const int32 Height, const int32 Width, UTexture2D* Texture)
+{
+	if (!Texture || !NumFrames)
+	{
+		return false;
+	}
+
+	// NumElements Per-Frame
+	const int32 NumElements = Vectors.Num() / NumFrames;
+
+	// Allocate PixelData.
+	TArray<typename TextureSettings::ColorType> Pixels;
+	Pixels.Init(TextureSettings::DefaultColor, Height * Width);
+
+	// Fillout Frame Data
+	for (int32 Frame = 0; Frame < NumFrames; ++Frame)
+	{
+		const int32 BlockStart = RowsPerFrame * Width * Frame;
+		
+		// Set Data.
+		for (int32 Index = 0; Index < NumElements; ++Index)
+		{
+			const V& Vector = Vectors[NumElements * Frame + Index];
+			typename TextureSettings::ColorType& Pixel = Pixels[BlockStart + Index];
+
+			VectorToColor<V, typename TextureSettings::ColorType>(Vector, Pixel);
+		}
+	}
+
+	// Write to Texture
+	return WriteToTexture<TextureSettings>(Texture, Height, Width, Pixels);
+}
+
+template<class TextureSettings>
+FORCEINLINE bool AnimToTexture_Private::WriteToTexture(
+	UTexture2D* Texture, 
+	const uint32 Height, 
+	const uint32 Width, 
+	const TArray<typename TextureSettings::ColorType>& Pixels)
+{
+	if (!Texture)
+	{
+		return false;
+	}
+
+	// ------------------------------------------------------------------------
+	// Get Texture Platform
+	FTexturePlatformData* PlatformData = Texture->GetPlatformData();
+	if (!PlatformData)
+	{	
+		PlatformData = new FTexturePlatformData();
+		Texture->SetPlatformData(PlatformData);
+	}
+	PlatformData->SizeX = Width;
+	PlatformData->SizeY = Height;
+	PlatformData->SetNumSlices(1);
+	PlatformData->PixelFormat = TextureSettings::PixelFormat;
+	
+	// ------------------------------------------------------------------------
+	// Get First MipMap
+	//
+	FTexture2DMipMap* Mip;
+	if (PlatformData->Mips.IsEmpty())
+	{
+		Mip = new FTexture2DMipMap();
+		PlatformData->Mips.Add(Mip);
+	}
+	else
+	{
+		Mip = &PlatformData->Mips[0];
+	}
+	Mip->SizeX = Width;
+	Mip->SizeY = Height;
+	
+	// ------------------------------------------------------------------------
+	// Lock the Mipmap data so it can be modified
+	Mip->BulkData.Lock(LOCK_READ_WRITE);
+
+	// Reallocate MipMap
+	uint8* TextureData = (uint8*)Mip->BulkData.Realloc(Width * Height * sizeof(typename TextureSettings::ColorType));
+
+	// Copy the pixel data into the Texture data	
+	const uint8* PixelsData = (uint8*)Pixels.GetData();
+	FMemory::Memcpy(TextureData, PixelsData, Width * Height * sizeof(typename TextureSettings::ColorType));
+
+	// Unlock data
+	Mip->BulkData.Unlock();
+
+	// Initialize a new texture
+	Texture->Source.Init(Width, Height, 1, 1, TextureSettings::TextureSourceFormat, PixelsData);
+	
+	// Set parameters
+	Texture->SRGB = 0;
+	Texture->Filter = TextureFilter::TF_Nearest;
+	Texture->CompressionSettings = TextureSettings::CompressionSettings;
+	Texture->MipGenSettings = TextureMipGenSettings::TMGS_NoMipmaps;
+
+	// Update and Mark to Save.
+	Texture->UpdateResource();
+	Texture->MarkPackageDirty();
+
+	return true;
+}
+
+
