Fix for static analysis warning in AnimToTextureBPLibrary
#rnx

diff --git a/Engine/Plugins/Animation/AnimToTexture/Source/AnimToTextureEditor/Private/AnimToTextureBPLibrary.cpp b/Engine/Plugins/Animation/AnimToTexture/Source/AnimToTextureEditor/Private/AnimToTextureBPLibrary.cpp
--- a/Engine/Plugins/Animation/AnimToTexture/Source/AnimToTextureEditor/Private/AnimToTextureBPLibrary.cpp
+++ b/Engine/Plugins/Animation/AnimToTexture/Source/AnimToTextureEditor/Private/AnimToTextureBPLibrary.cpp
@@ -98,10 +98,26 @@ void UAnimToTextureBPLibrary::AnimationToTexture(UAnimToTextureDataAsset* DataAs
 
 	// --------------------------------------------------------------------------
 
+	if (GEditor == nullptr)
+	{
+		return;
+	}
+
 	// Create Temp Actor
 	UWorld* World = GEditor->GetEditorWorldContext().World();
+
+	if (World == nullptr)
+	{
+		return;
+	}
+
 	AActor* Actor = World->SpawnActor<AActor>();
 
+	if (Actor == nullptr)
+	{
+		return;
+	}
+
 	// Create Temp LeaderMeshComponent
 	USkeletalMeshComponent* LeaderMeshComponent = nullptr;
 	if (DataAsset->GetLeaderSkeletalMesh())
@@ -117,6 +133,7 @@ void UAnimToTextureBPLibrary::AnimationToTexture(UAnimToTextureDataAsset* DataAs
 
 	// Create Temp SkeletalMesh Component
 	USkeletalMeshComponent* MeshComponent = NewObject<USkeletalMeshComponent>(Actor);
+	check(MeshComponent);
 	MeshComponent->SetSkeletalMesh(DataAsset->GetSkeletalMesh());
 	MeshComponent->SetForcedLOD(1); // Force to LOD0;
 	MeshComponent->RegisterComponent();
@@ -138,6 +155,11 @@ void UAnimToTextureBPLibrary::AnimationToTexture(UAnimToTextureDataAsset* DataAs
 		BaseMeshComponent = MeshComponent;
 	}
 	
+	if (BaseMeshComponent == nullptr)
+	{
+		return;
+	}
+
 	// --------------------------------------------------------------------------
 	// Set AnimBlueprint
 	UEvaluateSequenceAnimInstance* EvaluationAnimInstance = nullptr;
@@ -224,7 +246,7 @@ void UAnimToTextureBPLibrary::AnimationToTexture(UAnimToTextureDataAsset* DataAs
 		// Set AnimSequence
 		else
 		{
-			if (IsValid(EvaluationAnimInstance))
+			if (EvaluationAnimInstance)
 			{
 				EvaluationAnimInstance->SequenceToEvaluate = AnimSequence;
 				EvaluationAnimInstance->TimeToEvaluate = 0.0f;
@@ -271,7 +293,7 @@ void UAnimToTextureBPLibrary::AnimationToTexture(UAnimToTextureDataAsset* DataAs
 			SampleIndex++;
 
 			// Go To Time
-			if (IsValid(EvaluationAnimInstance))
+			if (EvaluationAnimInstance)
 			{
 				EvaluationAnimInstance->TimeToEvaluate = Time;
 			}
