Change Anim To Texture plugin to be experimental

diff --git a/Engine/Plugins/Animation/AnimToTexture/AnimToTexture.uplugin b/Engine/Plugins/Animation/AnimToTexture/AnimToTexture.uplugin
--- a/Engine/Plugins/Animation/AnimToTexture/AnimToTexture.uplugin
+++ b/Engine/Plugins/Animation/AnimToTexture/AnimToTexture.uplugin
@@ -11,8 +11,8 @@
 	"MarketplaceURL": "",
 	"SupportURL": "",
 	"CanContainContent": true,
-	"IsBetaVersion": true,
-	"IsExperimentalVersion": false,
+	"IsBetaVersion": false,
+	"IsExperimentalVersion": true,
 	"Installed": false,
 	"Modules": [
 		{
