AnimToTexture:
- Using CompatibleSkeletons instead of LeaderComponent
- Removed untested option: AnimBlueprint
- Removed project specific logic. CreateUVs, CreateTextures.
- Setting UVChannel 1 as default
- Added property descriptions
- Removed commented out code

diff --git a/Engine/Plugins/Animation/AnimToTexture/Source/AnimToTexture/Private/AnimToTextureDataAsset.cpp b/Engine/Plugins/Animation/AnimToTexture/Source/AnimToTexture/Private/AnimToTextureDataAsset.cpp
--- a/Engine/Plugins/Animation/AnimToTexture/Source/AnimToTexture/Private/AnimToTextureDataAsset.cpp
+++ b/Engine/Plugins/Animation/AnimToTexture/Source/AnimToTexture/Private/AnimToTextureDataAsset.cpp
@@ -29,8 +29,6 @@ namespace AnimToTextureParamNames
 	static const FName RowsPerFrame = TEXT("RowsPerFrame");
 	static const FName BoneWeightRowsPerFrame = TEXT("BoneWeightsRowsPerFrame");
 	static const FName NumFrames = TEXT("NumFrames (S)");
-	//static const FName UVChannel = TEXT("UVChannel");
-	//static const FName AnimateSwitch = TEXT("Animate (B)");
 	static const FName VertexPositionTexture = TEXT("PositionTexture");
 	static const FName VertexNormalTexture = TEXT("NormalTexture");
 	static const FName BonePositionTexture = TEXT("BonePositionTexture");
@@ -45,8 +43,6 @@ FAnimToTextureMaterialParamNames::FAnimToTextureMaterialParamNames()
 	RowsPerFrame = AnimToTextureParamNames::RowsPerFrame;
 	BoneWeightRowsPerFrame = AnimToTextureParamNames::BoneWeightRowsPerFrame;
 	NumFrames = AnimToTextureParamNames::NumFrames;
-	//UVChannel = AnimToTextureParamNames::UVChannel;
-	//AnimateSwitch = AnimToTextureParamNames::AnimateSwitch;
 	VertexPositionTexture = AnimToTextureParamNames::VertexPositionTexture;
 	VertexNormalTexture = AnimToTextureParamNames::VertexNormalTexture;
 	BonePositionTexture = AnimToTextureParamNames::BonePositionTexture;
diff --git a/Engine/Plugins/Animation/AnimToTexture/Source/AnimToTexture/Public/AnimToTextureDataAsset.h b/Engine/Plugins/Animation/AnimToTexture/Source/AnimToTexture/Public/AnimToTextureDataAsset.h
--- a/Engine/Plugins/Animation/AnimToTexture/Source/AnimToTexture/Public/AnimToTextureDataAsset.h
+++ b/Engine/Plugins/Animation/AnimToTexture/Source/AnimToTexture/Public/AnimToTextureDataAsset.h
@@ -16,20 +16,10 @@ struct ANIMTOTEXTURE_API FAnimToTextureMaterialParamNames
 {
 	GENERATED_BODY()
 
-	//
-	// Static Switch Parameters
-	//
-
-	// UPROPERTY(BlueprintReadWrite, EditAnywhere)
-	// FName AnimateSwitch;
-
 	//
 	// Scalar Parameters
 	//
 
-	//UPROPERTY(EditAnywhere)
-	//FName UVChannel;
-
 	UPROPERTY(BlueprintReadWrite, Category = Default, EditAnywhere)
 	FName RowsPerFrame;
 
@@ -81,12 +71,6 @@ enum class EAnimToTextureMode : uint8
 	Bone,
 };
 
-//UENUM(Blueprintable)
-//enum class EAnimToTextureNumInfluences : uint8
-//{
-//	One, Two, Four,
-//};
-
 UENUM(Blueprintable)
 enum class EAnimToTextureBonePrecision : uint8
 {
@@ -96,34 +80,6 @@ enum class EAnimToTextureBonePrecision : uint8
 	SixteenBits,
 };
 
-
-//USTRUCT(Blueprintable)
-//struct FSkeletalMeshInfo
-//{
-//	GENERATED_BODY()
-//
-//	UPROPERTY(EditAnywhere)
-//	USkeletalMesh* SkeletalMesh = nullptr;
-//
-//	UPROPERTY(EditAnywhere)
-//	int32 LODIndex = 0;
-//};
-
-//USTRUCT(Blueprintable)
-//struct FStaticMeshInfo
-//{
-//	GENERATED_BODY()
-//
-//	UPROPERTY(EditAnywhere)
-//	UStaticMesh* StaticMesh = nullptr;
-//
-//	UPROPERTY(EditAnywhere)
-//	int32 LODIndex = 0;
-//
-//	UPROPERTY(EditAnywhere)
-//	int32 UVChannel = 1;
-//};
-
 USTRUCT(Blueprintable)
 struct FAnimSequenceInfo
 {
@@ -170,79 +126,120 @@ class ANIMTOTEXTURE_API UAnimToTextureDataAsset : public UPrimaryDataAsset
 public:
 	GENERATED_BODY()
 
-	// ------------------------------------------------------
-	// Skeletal Mesh
-
+	/**
+	* SkeletalMesh to bake animations from.
+	*/
 	UPROPERTY(EditAnywhere, BlueprintReadWrite, Category = "SkeletalMesh", meta = (AssetBundles = "Client"))
 	TSoftObjectPtr<USkeletalMesh> SkeletalMesh;
 
+	/**
+	* SkeletalMesh LOD to bake.
+	*/
 	UPROPERTY(EditAnywhere, BlueprintReadWrite, Category = "SkeletalMesh", Meta = (DisplayName = "LODIndex"))
 	int32 SkeletalLODIndex = 0;
 
-	// ------------------------------------------------------
-	// Static Mesh
-
+	/**
+	* StaticMesh to bake to.
+	*/
 	UPROPERTY(EditAnywhere, BlueprintReadWrite, Category = "StaticMesh", meta = (AssetBundles = "Client"))
 	TSoftObjectPtr<UStaticMesh> StaticMesh;
 
+	/**
+	* StaticMesh LOD to bake to.
+	*/
 	UPROPERTY(EditAnywhere, BlueprintReadWrite, Category = "StaticMesh", Meta = (DisplayName = "LODIndex"))
 	int32 StaticLODIndex = 0;
 
+	/**
+	* StaticMesh UVChannel Index for storing vertex information.
+	* Make sure this index does not conflict with the Lightmap UV Index.
+	*/
 	UPROPERTY(EditAnywhere, BlueprintReadWrite, Category = "StaticMesh")
 	int32 UVChannel = 1;
 
 	// ------------------------------------------------------
 	// Texture
 
+	/**
+	* Max resolution of the texture.
+	* A smaller size will be used if the data fits.
+	*/
 	UPROPERTY(EditAnywhere, BlueprintReadWrite, Category = "Texture")
 	int32 MaxHeight = 4096;
 
+	/**
+	* Max resolution of the texture.
+	* A smaller size will be used if the data fits.
+	*/
 	UPROPERTY(EditAnywhere, BlueprintReadWrite, Category = "Texture")
 	int32 MaxWidth = 4096;
 
+	/**
+	* Enforce Power Of Two on texture resolutions.
+	*/
 	UPROPERTY(EditAnywhere, BlueprintReadWrite, Category = "Texture")
 	bool bEnforcePowerOfTwo = false;
 
+	/**
+	* Storage Mode.
+	* Vertex: will store per-vertex position and normal.
+	* Bone: Will store per-bone position and rotation and per-vertex bone weight. 
+	        This is the preferred method if meshes share skeleton.
+	*/
 	UPROPERTY(EditAnywhere, BlueprintReadWrite, Category = "Texture")
 	EAnimToTextureMode Mode;
 
+	/**
+	* Texture for storing vertex positions
+	* This is only used on Vertex Mode
+	*/
 	UPROPERTY(EditAnywhere, BlueprintReadWrite, Category = "Texture|Vertex", meta = (EditCondition = "Mode == EAnimToTextureMode::Vertex"))
 	TSoftObjectPtr<UTexture2D> VertexPositionTexture;
 
+	/**
+	* Texture for storing vertex normals
+	* This is only used on Vertex Mode
+	*/
 	UPROPERTY(EditAnywhere, BlueprintReadWrite, Category = "Texture|Vertex", meta = (EditCondition = "Mode == EAnimToTextureMode::Vertex"))
 	TSoftObjectPtr<UTexture2D> VertexNormalTexture;
 
-	// UPROPERTY(EditAnywhere, BlueprintReadWrite, Category = "Texture|Bone", meta = (EditCondition = "Mode == EAnimToTextureMode::Bone"))
-	// EAnimToTextureNumInfluences NumInfluences;
-	
+	/**
+	* Texture Precision for: BonePosition and BoneRotations.
+	*/
+	UPROPERTY(EditAnywhere, BlueprintReadWrite, Category = "Texture|Bone", meta = (EditCondition = "Mode == EAnimToTextureMode::Bone"))
+	EAnimToTextureBonePrecision BonePrecision = EAnimToTextureBonePrecision::EightBits;
+
+	/**
+	* Texture for storing bone positions
+	* This is only used on Bone Mode
+	*/
 	UPROPERTY(EditAnywhere, BlueprintReadWrite, Category = "Texture|Bone", meta = (EditCondition = "Mode == EAnimToTextureMode::Bone"))
 	TSoftObjectPtr<UTexture2D> BonePositionTexture;
 
+	/**
+	* Texture for storing bone rotations
+	* This is only used on Bone Mode
+	*/
 	UPROPERTY(EditAnywhere, BlueprintReadWrite, Category = "Texture|Bone", meta = (EditCondition = "Mode == EAnimToTextureMode::Bone"))
 	TSoftObjectPtr<UTexture2D> BoneRotationTexture;
 
+	/**
+	* Texture for storing vertex/bone weighting
+	* This is only used on Bone Mode
+	*/
 	UPROPERTY(EditAnywhere, BlueprintReadWrite, Category = "Texture|Bone", meta = (EditCondition = "Mode == EAnimToTextureMode::Bone"))
 	TSoftObjectPtr<UTexture2D> BoneWeightTexture;
 
-	UPROPERTY(EditAnywhere, BlueprintReadWrite, Category = "Texture|Bone", meta = (EditCondition = "Mode == EAnimToTextureMode::Bone"))
-	EAnimToTextureBonePrecision PositionAndRotationPrecision = EAnimToTextureBonePrecision::EightBits;
-	
 	// ------------------------------------------------------
 	// Animation
 
-	 /** This Mesh will be used as LeaderPose. Animations must use same Skeleton than this SkeletalMesh */
-	UPROPERTY(EditAnywhere, BlueprintReadWrite, Category = "Animation")
-	TSoftObjectPtr<USkeletalMesh> LeaderSkeletalMesh;
-
-	/** Bone used for Rigid Binding. The bone needs to be part of the RawBones. 
-	*   Sockets and VirtualBones are not supported.
+	/** 
+	* Bone used for Rigid Binding. The bone needs to be part of the RawBones. 
+	* Sockets and VirtualBones are not supported.
 	*/
-	UPROPERTY(EditAnywhere, BlueprintReadWrite, Category = "Animation")
+	UPROPERTY(EditAnywhere, BlueprintReadWrite, Category = "Animation", meta = (EditCondition = "Mode == EAnimToTextureMode::Bone"))
 	FName AttachToSocket;
 
-	UPROPERTY(EditAnywhere, BlueprintReadWrite, Category = "Animation")
-	TSubclassOf<UAnimInstance> AnimInstanceClass;
-
 	UPROPERTY(EditAnywhere, BlueprintReadWrite, Category = "Animation")
 	float SampleRate = 30.0f;
 
@@ -311,7 +308,6 @@ public:
 
 	AnimToTextureDataAsset_ASSET_ACCESSOR(UStaticMesh, StaticMesh);
 	AnimToTextureDataAsset_ASSET_ACCESSOR(USkeletalMesh, SkeletalMesh);
-	AnimToTextureDataAsset_ASSET_ACCESSOR(USkeletalMesh, LeaderSkeletalMesh);
 	AnimToTextureDataAsset_ASSET_ACCESSOR(UTexture2D, VertexPositionTexture);
 	AnimToTextureDataAsset_ASSET_ACCESSOR(UTexture2D, VertexNormalTexture);
 	AnimToTextureDataAsset_ASSET_ACCESSOR(UTexture2D, BonePositionTexture);
@@ -324,9 +320,6 @@ public:
 	UFUNCTION(BlueprintPure, Category = Default, meta = (DisplayName = "Get Skeletal Mesh"))
 	USkeletalMesh* BP_GetSkeletalMesh() { return GetSkeletalMesh(); }
 
-	UFUNCTION(BlueprintPure, Category = Default, meta = (DisplayName = "Get Leader Skeletal Mesh"))
-	USkeletalMesh* BP_GetLeaderSkeletalMesh() { return GetLeaderSkeletalMesh(); }
-
 	UFUNCTION(BlueprintPure, Category = Default, meta = (DisplayName = "Get Bone Position Texture"))
 	UTexture2D* BP_GetBonePositionTexture() { return GetBonePositionTexture(); }
 
@@ -340,19 +333,17 @@ public:
 FORCEINLINE void UAnimToTextureDataAsset::Reset()
 {
 	// Common Info.
-	//this->NumVertices = 0;
 	this->NumFrames = 0;
 	this->Animations.Reset();
 
 	// Vertex Info
 	this->VertexRowsPerFrame = 1;
-	this->VertexMinBBox = FVector::ZeroVector; // { TNumericLimits<float>::Max(), TNumericLimits<float>::Max(), TNumericLimits<float>::Max() };
+	this->VertexMinBBox = FVector::ZeroVector;
 	this->VertexSizeBBox = FVector::ZeroVector;
 
 	// Bone Info
-	//this->NumBones = 0;
 	this->BoneRowsPerFrame = 1;
 	this->BoneWeightRowsPerFrame = 1;
-	this->BoneMinBBox = FVector::ZeroVector;  // { TNumericLimits<float>::Max(), TNumericLimits<float>::Max(), TNumericLimits<float>::Max() };
+	this->BoneMinBBox = FVector::ZeroVector; 
 	this->BoneSizeBBox = FVector::ZeroVector;
 };
\ No newline at end of file
diff --git a/Engine/Plugins/Animation/AnimToTexture/Source/AnimToTextureEditor/Private/AnimToTextureBPLibrary.cpp b/Engine/Plugins/Animation/AnimToTexture/Source/AnimToTextureEditor/Private/AnimToTextureBPLibrary.cpp
--- a/Engine/Plugins/Animation/AnimToTexture/Source/AnimToTextureEditor/Private/AnimToTextureBPLibrary.cpp
+++ b/Engine/Plugins/Animation/AnimToTexture/Source/AnimToTextureEditor/Private/AnimToTextureBPLibrary.cpp
@@ -31,8 +31,7 @@ UAnimToTextureBPLibrary::UAnimToTextureBPLibrary(const FObjectInitializer& Objec
 
 }
 
-void UAnimToTextureBPLibrary::AnimationToTexture(UAnimToTextureDataAsset* DataAsset,
-	const FTransform RootTransform, const bool bCreateTextures, const bool bCreateUVChannel)
+void UAnimToTextureBPLibrary::AnimationToTexture(UAnimToTextureDataAsset* DataAsset, const FTransform RootTransform)
 {
 	if (!DataAsset)
 	{
@@ -118,66 +117,16 @@ void UAnimToTextureBPLibrary::AnimationToTexture(UAnimToTextureDataAsset* DataAs
 		return;
 	}
 
-	// Create Temp LeaderMeshComponent
-	USkeletalMeshComponent* LeaderMeshComponent = nullptr;
-	if (DataAsset->GetLeaderSkeletalMesh())
-	{
-		LeaderMeshComponent = NewObject<USkeletalMeshComponent>(Actor);
-		LeaderMeshComponent->SetSkeletalMesh(DataAsset->GetLeaderSkeletalMesh());
-		LeaderMeshComponent->SetForcedLOD(1); // Force to LOD0
-		LeaderMeshComponent->SetAnimationMode(EAnimationMode::AnimationSingleNode);
-		LeaderMeshComponent->SetUpdateAnimationInEditor(true);
-		LeaderMeshComponent->VisibilityBasedAnimTickOption = EVisibilityBasedAnimTickOption::AlwaysTickPoseAndRefreshBones;		
-		LeaderMeshComponent->RegisterComponent();
-	};
-
 	// Create Temp SkeletalMesh Component
 	USkeletalMeshComponent* MeshComponent = NewObject<USkeletalMeshComponent>(Actor);
 	check(MeshComponent);
 	MeshComponent->SetSkeletalMesh(DataAsset->GetSkeletalMesh());
 	MeshComponent->SetForcedLOD(1); // Force to LOD0;
+	MeshComponent->SetAnimationMode(EAnimationMode::AnimationSingleNode);
+	MeshComponent->SetUpdateAnimationInEditor(true);
+	MeshComponent->VisibilityBasedAnimTickOption = EVisibilityBasedAnimTickOption::AlwaysTickPoseAndRefreshBones;
 	MeshComponent->RegisterComponent();
-	
-	USkeletalMeshComponent* BaseMeshComponent = nullptr;
-	if (LeaderMeshComponent)
-	{
-		MeshComponent->SetupAttachment(LeaderMeshComponent);
-		MeshComponent->SetLeaderPoseComponent(LeaderMeshComponent);
 
-		BaseMeshComponent = LeaderMeshComponent;
-	}
-	else
-	{
-		MeshComponent->SetAnimationMode(EAnimationMode::AnimationSingleNode);
-		MeshComponent->SetUpdateAnimationInEditor(true);
-		MeshComponent->VisibilityBasedAnimTickOption = EVisibilityBasedAnimTickOption::AlwaysTickPoseAndRefreshBones;
-		
-		BaseMeshComponent = MeshComponent;
-	}
-	
-	if (BaseMeshComponent == nullptr)
-	{
-		return;
-	}
-
-	// --------------------------------------------------------------------------
-	// Set AnimBlueprint
-	UEvaluateSequenceAnimInstance* EvaluationAnimInstance = nullptr;
-	if (DataAsset->AnimInstanceClass != nullptr)
-	{
-		// Set Anim Blueprint
-		BaseMeshComponent->SetAnimInstanceClass(DataAsset->AnimInstanceClass);
-		EvaluationAnimInstance = Cast<UEvaluateSequenceAnimInstance>(BaseMeshComponent->GetAnimInstance());
-	}
-
-	// Get Mapping between Follower and Leader
-	TArray<int32> FollowerBoneMap;
-	if (DataAsset->GetLeaderSkeletalMesh() &&
-		DataAsset->Mode == EAnimToTextureMode::Bone)
-	{
-		AnimToTexture_Private::GetFollowerBoneMap(DataAsset->GetLeaderSkeletalMesh(), DataAsset->GetSkeletalMesh(), FollowerBoneMap);
-	}
-	
 	// ---------------------------------------------------------------------------
 	// Get Reference Skeleton Transforms
 	//
@@ -192,15 +141,15 @@ void UAnimToTextureBPLibrary::AnimationToTexture(UAnimToTextureDataAsset* DataAs
 	if (DataAsset->Mode == EAnimToTextureMode::Bone)
 	{
 		// Get Number of RawBones (no virtual)
-		NumBones = AnimToTexture_Private::GetNumBones(BaseMeshComponent->GetSkeletalMeshAsset());
+		NumBones = AnimToTexture_Private::GetNumBones(MeshComponent->GetSkeletalMeshAsset());
 
 		// Get Raw Ref Bone (no virtual)
 		TArray<FTransform> RefBoneTransforms;
-		AnimToTexture_Private::GetRefBoneTransforms(BaseMeshComponent->GetSkeletalMeshAsset(), RefBoneTransforms);
+		AnimToTexture_Private::GetRefBoneTransforms(MeshComponent->GetSkeletalMeshAsset(), RefBoneTransforms);
 		AnimToTexture_Private::DecomposeTransformations(RefBoneTransforms, BoneRefPositions, BoneRefRotations);
 
 		// Get Bone Names (no virtual)
-		AnimToTexture_Private::GetBoneNames(BaseMeshComponent->GetSkeletalMeshAsset(), BoneNames);
+		AnimToTexture_Private::GetBoneNames(MeshComponent->GetSkeletalMeshAsset(), BoneNames);
 
 		// Make sure array sizes are correct.
 		check(RefBoneTransforms.Num() == NumBones && BoneNames.Num() == NumBones);
@@ -237,24 +186,16 @@ void UAnimToTextureBPLibrary::AnimationToTexture(UAnimToTextureDataAsset* DataAs
 			continue;
 		}
 		
-		// Make sure SkelMesh and AnimSequence use same Skeleton
-		if (AnimSequence->GetSkeleton() != BaseMeshComponent->GetSkeletalMeshAsset()->GetSkeleton())
+		// Make sure SkeletalMesh is compatible with AnimSequence
+		if (!MeshComponent->GetSkeletalMeshAsset()->GetSkeleton()->IsCompatible(AnimSequence->GetSkeleton()))
 		{
-			UE_LOG(LogTemp, Warning, TEXT("Invalid AnimSequence: %s for given SkeletalMesh: %s"), *AnimSequence->GetFName().ToString(), *BaseMeshComponent->GetSkeletalMeshAsset()->GetFName().ToString());
+			UE_LOG(LogTemp, Warning, TEXT("Invalid AnimSequence: %s for given SkeletalMesh: %s"), *AnimSequence->GetFName().ToString(), *MeshComponent->GetSkeletalMeshAsset()->GetFName().ToString());
 			continue;
 		}
 		// Set AnimSequence
 		else
-		{
-			if (EvaluationAnimInstance)
-			{
-				EvaluationAnimInstance->SequenceToEvaluate = AnimSequence;
-				EvaluationAnimInstance->TimeToEvaluate = 0.0f;
-			}
-			else
-			{
-				BaseMeshComponent->SetAnimation(AnimSequence);
-			}
+		{			
+			MeshComponent->SetAnimation(AnimSequence);
 		}
 
 		// -----------------------------------------------------------------------------------
@@ -293,20 +234,12 @@ void UAnimToTextureBPLibrary::AnimationToTexture(UAnimToTextureDataAsset* DataAs
 			SampleIndex++;
 
 			// Go To Time
-			if (EvaluationAnimInstance)
-			{
-				EvaluationAnimInstance->TimeToEvaluate = Time;
-			}
-			else
-			{
-				BaseMeshComponent->SetPosition(Time);
-			}
+			MeshComponent->SetPosition(Time);
 
 			// Update SkelMesh Animation.
-			BaseMeshComponent->TickAnimation(0.0f, false /*bNeedsValidRootMotion*/);
-			// BaseMeshComponent->TickComponent(0.0f, ELevelTick::LEVELTICK_All, nullptr);
-			BaseMeshComponent->RefreshBoneTransforms(nullptr /*TickFunction*/);
-			BaseMeshComponent->RefreshFollowerComponents();
+			MeshComponent->TickAnimation(0.0f, false /*bNeedsValidRootMotion*/);
+			// MeshComponent->TickComponent(0.0f, ELevelTick::LEVELTICK_All, nullptr);
+			MeshComponent->RefreshBoneTransforms(nullptr /*TickFunction*/);
 			
 			// ---------------------------------------------------------------------------
 			// Store Vertex Deltas & Normals.
@@ -339,17 +272,15 @@ void UAnimToTextureBPLibrary::AnimationToTexture(UAnimToTextureDataAsset* DataAs
 			{
 				// Get Relative Transforms
 				// Note: Size is of Raw bones in SkeletalMesh. These are the original/raw bones of the asset, without Virtual Bones.
-				// Note: You cannot call CacheRefToLocalMatrices on a FollowerComponent
 				TArray<FMatrix44f> RefToLocals;
-				BaseMeshComponent->CacheRefToLocalMatrices(RefToLocals);
+				MeshComponent->CacheRefToLocalMatrices(RefToLocals);
 				
 				// check size
 				check(RefToLocals.Num() == NumBones);
 
 				// Get Component Space Transforms
 				// Note returns all transforms, including VirtualBones
-				// Note: You cannot call GetComponentSpaceTransforms on a FollowerComponent
-				const TArray<FTransform>& CompSpaceTransforms = BaseMeshComponent->GetComponentSpaceTransforms();
+				const TArray<FTransform>& CompSpaceTransforms = MeshComponent->GetComponentSpaceTransforms();
 				check(CompSpaceTransforms.Num() >= RefToLocals.Num());
 
 				for (int32 BoneIndex = 0; BoneIndex < NumBones; ++BoneIndex)
@@ -391,13 +322,6 @@ void UAnimToTextureBPLibrary::AnimationToTexture(UAnimToTextureDataAsset* DataAs
 	MeshComponent->UnregisterComponent();
 	MeshComponent->DestroyComponent();
 
-	// Destroy Leader Component
-	if (LeaderMeshComponent)
-	{
-		LeaderMeshComponent->UnregisterComponent();
-		LeaderMeshComponent->DestroyComponent();
-	}
-	
 	Actor->Destroy();
 	
 	// ---------------------------------------------------------------------------
@@ -430,18 +354,12 @@ void UAnimToTextureBPLibrary::AnimationToTexture(UAnimToTextureDataAsset* DataAs
 			NormalizedVertexDeltas, NormalizedVertexNormals);
 
 		// Write Textures
-		if (bCreateTextures)
-		{
-			AnimToTexture_Private::WriteVectorsToTexture<FVector3f, AnimToTexture_Private::FLowPrecision>(NormalizedVertexDeltas, DataAsset->NumFrames, DataAsset->VertexRowsPerFrame, Height, Width, DataAsset->GetVertexPositionTexture());
-			AnimToTexture_Private::WriteVectorsToTexture<FVector3f, AnimToTexture_Private::FLowPrecision>(NormalizedVertexNormals, DataAsset->NumFrames, DataAsset->VertexRowsPerFrame, Height, Width, DataAsset->GetVertexNormalTexture());
-		}
+		AnimToTexture_Private::WriteVectorsToTexture<FVector3f, AnimToTexture_Private::FLowPrecision>(NormalizedVertexDeltas, DataAsset->NumFrames, DataAsset->VertexRowsPerFrame, Height, Width, DataAsset->GetVertexPositionTexture());
+		AnimToTexture_Private::WriteVectorsToTexture<FVector3f, AnimToTexture_Private::FLowPrecision>(NormalizedVertexNormals, DataAsset->NumFrames, DataAsset->VertexRowsPerFrame, Height, Width, DataAsset->GetVertexNormalTexture());
 
+	
 		// Add Vertex UVChannel
-		if (bCreateUVChannel)
-		{
-			CreateUVChannel(DataAsset->GetStaticMesh(), DataAsset->StaticLODIndex, DataAsset->UVChannel,
-				Height, Width);
-		}
+		CreateUVChannel(DataAsset->GetStaticMesh(), DataAsset->StaticLODIndex, DataAsset->UVChannel, Height, Width);
 	}
 
 	// ---------------------------------------------------------------------------
@@ -451,7 +369,7 @@ void UAnimToTextureBPLibrary::AnimationToTexture(UAnimToTextureDataAsset* DataAs
 		// Find Best Resolution for Bone Data
 		int32 Height, Width;
 
-		// NOTE: If NumBones are > 256, you might need to use LeaderPose
+		// NOTE: If NumBones are > 256, you might need to use a simpler skeleton
 
 		// Note we are adding +1 frame for the ref pose
 		if (!FindBestResolution(DataAsset->NumFrames + 1, NumBones, 
@@ -463,7 +381,6 @@ void UAnimToTextureBPLibrary::AnimationToTexture(UAnimToTextureDataAsset* DataAs
 		}
 
 		// Write Bone Position and Rotation Textures
-		if (bCreateTextures)
 		{
 			// Normalize Bone Data
 			TArray<FVector3f> NormalizedBonePositions;
@@ -474,7 +391,7 @@ void UAnimToTextureBPLibrary::AnimationToTexture(UAnimToTextureDataAsset* DataAs
 				NormalizedBonePositions, NormalizedBoneRotations);
 
 			// Write Textures
-			if (DataAsset->PositionAndRotationPrecision == EAnimToTextureBonePrecision::SixteenBits)
+			if (DataAsset->BonePrecision == EAnimToTextureBonePrecision::SixteenBits)
 			{
 				AnimToTexture_Private::WriteVectorsToTexture<FVector3f, AnimToTexture_Private::FHighPrecision>(NormalizedBonePositions, DataAsset->NumFrames + 1, DataAsset->BoneRowsPerFrame, Height, Width, DataAsset->GetBonePositionTexture());
 				AnimToTexture_Private::WriteVectorsToTexture<FVector4, AnimToTexture_Private::FHighPrecision>(NormalizedBoneRotations, DataAsset->NumFrames + 1, DataAsset->BoneRowsPerFrame, Height, Width, DataAsset->GetBoneRotationTexture());
@@ -498,7 +415,6 @@ void UAnimToTextureBPLibrary::AnimationToTexture(UAnimToTextureDataAsset* DataAs
 		}
 
 		// Write Weights Texture
-		if (bCreateTextures)
 		{
 			TArray<AnimToTexture_Private::TVertexSkinWeight<4>> SkinWeights;
 
@@ -510,18 +426,6 @@ void UAnimToTextureBPLibrary::AnimationToTexture(UAnimToTextureDataAsset* DataAs
 					*DataAsset->GetSkeletalMesh(), DataAsset->SkeletalLODIndex,
 					StaticToSkelMapping, 
 					SkinWeights);
-
-				// Remap MeshBoneIndices to FollowerBoneMap
-				if (DataAsset->GetLeaderSkeletalMesh())
-				{
-					for (AnimToTexture_Private::TVertexSkinWeight<4>& SkinWeight : SkinWeights)
-					{
-						SkinWeight.MeshBoneIndices[0] = FollowerBoneMap[SkinWeight.MeshBoneIndices[0]];
-						SkinWeight.MeshBoneIndices[1] = FollowerBoneMap[SkinWeight.MeshBoneIndices[1]];
-						SkinWeight.MeshBoneIndices[2] = FollowerBoneMap[SkinWeight.MeshBoneIndices[2]];
-						SkinWeight.MeshBoneIndices[3] = FollowerBoneMap[SkinWeight.MeshBoneIndices[3]];
-					}
-				}
 			}
 			// If Valid Socket, set all influences to same index.
 			else
@@ -541,11 +445,7 @@ void UAnimToTextureBPLibrary::AnimationToTexture(UAnimToTextureDataAsset* DataAs
 		}
 
 		// Add Vertex UVChannel
-		if (bCreateUVChannel)
-		{
-			CreateUVChannel(DataAsset->GetStaticMesh(), DataAsset->StaticLODIndex, DataAsset->UVChannel,
-				Height, Width);
-		}
+		CreateUVChannel(DataAsset->GetStaticMesh(), DataAsset->StaticLODIndex, DataAsset->UVChannel, Height, Width);
 	}
 
 	// ---------------------------------------------------------------------------
diff --git a/Engine/Plugins/Animation/AnimToTexture/Source/AnimToTextureEditor/Public/AnimToTextureBPLibrary.h b/Engine/Plugins/Animation/AnimToTexture/Source/AnimToTextureEditor/Public/AnimToTextureBPLibrary.h
--- a/Engine/Plugins/Animation/AnimToTexture/Source/AnimToTextureEditor/Public/AnimToTextureBPLibrary.h
+++ b/Engine/Plugins/Animation/AnimToTexture/Source/AnimToTextureEditor/Public/AnimToTextureBPLibrary.h
@@ -26,33 +26,33 @@ class UAnimToTextureBPLibrary : public UBlueprintFunctionLibrary
 
 	/**
 	* Bakes Animation Data into Textures.
-	* Position & Normals.
 	*/
-	UFUNCTION(BlueprintCallable, meta = (Category = "Vertex Animation"))
-	static void AnimationToTexture(UAnimToTextureDataAsset* DataAsset,
-		const FTransform RootTransform, const bool bCreateTextures=true, const bool bCreateUVChannel=true);
+	UFUNCTION(BlueprintCallable, meta = (Category = "AnimToTexture"))
+	static void AnimationToTexture(UAnimToTextureDataAsset* DataAsset, const FTransform RootTransform);
 
-	/* Utility for converting SkeletalMesh into a StaticMesh */
-	UFUNCTION(BlueprintCallable, Category = "Vertex Animation")
+	/** 
+	* Utility for converting SkeletalMesh into a StaticMesh
+	*/
+	UFUNCTION(BlueprintCallable, Category = "AnimToTexture")
 	static UStaticMesh* ConvertSkeletalMeshToStaticMesh(USkeletalMesh* SkeletalMesh, const FString PackageName, const int32 LODIndex = -1);
 
 	/**
 	* Updates a material's parameters to match those of an animToTexture data asset
 	*/
-	UFUNCTION(BlueprintCallable, meta = (Category = "Vertex Animation"))
+	UFUNCTION(BlueprintCallable, meta = (Category = "AnimToTexture"))
 	static void UpdateMaterialInstanceFromDataAsset(UAnimToTextureDataAsset* DataAsset, class UMaterialInstanceConstant* MaterialInstance, 
 		const FAnimToTextureMaterialParamNames& ParamNames, const EMaterialParameterAssociation MaterialParameterAssociation = EMaterialParameterAssociation::GlobalParameter);
 
 	/**
 	 * Replaces material layer function
 	 */
-	UFUNCTION(BlueprintCallable, meta = (Category = "Vertex Animation"))
+	UFUNCTION(BlueprintCallable, meta = (Category = "AnimToTexture"))
 	static void UpdateMaterialLayerFunction(
 		class UMaterialInstanceConstant* MaterialInstance, 
 		class UMaterialFunctionInterface* OldMaterialFunction, 
 		class UMaterialFunctionInterface* NewMaterialFunction);
 
-	UFUNCTION(BlueprintCallable, meta = (Category = "Vertex Animation"))
+	UFUNCTION(BlueprintCallable, meta = (Category = "AnimToTexture"))
 	static void SetStaticMeshBoundsExtensions(UStaticMesh* StaticMesh, const FVector& PositiveBoundsExtension, const FVector& NegativeBoundsExtension);
 
 #endif // WITH_EDITOR
diff --git a/Engine/Plugins/Animation/AnimToTexture/Source/AnimToTextureEditor/Public/AnimToTextureUtils.h b/Engine/Plugins/Animation/AnimToTexture/Source/AnimToTextureEditor/Public/AnimToTextureUtils.h
--- a/Engine/Plugins/Animation/AnimToTexture/Source/AnimToTextureEditor/Public/AnimToTextureUtils.h
+++ b/Engine/Plugins/Animation/AnimToTexture/Source/AnimToTextureEditor/Public/AnimToTextureUtils.h
@@ -61,8 +61,6 @@ namespace AnimToTexture_Private
 	template<class V /* FVector / FVector4 / FIntVector */, class C /* FColor / FLinearColor */>
 	void VectorToColor(const V& Vector, C& Color);
 	
-	//template<class C /* FColor / FLinearColor */, class T /* uint8 / float */>
-	//void ColorArrayToData(const TArray<C>& Colors, TArray<T>& Data);
 }
 
 // ----------------------------------------------------------------------------
@@ -164,40 +162,6 @@ FORCEINLINE void AnimToTexture_Private::VectorToColor(const FVector4& Vector, FV
 	Color.W = FMath::RoundToInt(FMath::Clamp(Vector.W, 0.0f, 1.0f) * MAX_uint16);
 }
 
-/*
-template<>
-FORCEINLINE void AnimToTexture_Private::ColorArrayToData(const TArray<FColor>& Colors, TArray<uint8>& Data)
-{
-	// Resize Array
-	Data.Init(0, Colors.Num() * 4);
-
-	for (int32 Index = 0; Index < Colors.Num(); ++Index)
-	{
-		// NOTE: format is BGRA (swaping R <-> B)
-		Data[Index * 4 + 0] = Colors[Index].B; // B
-		Data[Index * 4 + 1] = Colors[Index].G; // G
-		Data[Index * 4 + 2] = Colors[Index].R; // R
-		Data[Index * 4 + 3] = Colors[Index].A; // A
-	}
-}
-
-template<>
-FORCEINLINE void AnimToTexture_Private::ColorArrayToData(const TArray<FLinearColor>& Colors, TArray<float>& Data)
-{
-	// Resize Array
-	Data.Init(0.0f, Colors.Num() * 4);
-
-	for (int32 Index = 0; Index < Colors.Num(); ++Index)
-	{
-		// NOTE: format is RGBA
-		Data[Index * 4 + 0] = Colors[Index].R; // R
-		Data[Index * 4 + 1] = Colors[Index].G; // G
-		Data[Index * 4 + 2] = Colors[Index].B; // B
-		Data[Index * 4 + 3] = Colors[Index].A; // A
-	}
-}
-*/
-
 
 template<class V, class TextureSettings>
 FORCEINLINE bool AnimToTexture_Private::WriteVectorsToTexture(const TArray<V>& Vectors,
