AnimToTexture :
- Improved UpdateMaterialInstanceFromDataAsset, setting UVIndex, Animate and NumBoneInfluences
- Renamed some Material params.
- Fixed StaticMesh bounds
- Setting StaticMesh FullPrecisionUVs only to given LOD
- Forcing PowerOfTwo for 16bit precision
- Added 16bit precision for Vertex Mode

diff --git a/Engine/Plugins/Animation/AnimToTexture/Source/AnimToTexture/Private/AnimToTextureDataAsset.cpp b/Engine/Plugins/Animation/AnimToTexture/Source/AnimToTexture/Private/AnimToTextureDataAsset.cpp
--- a/Engine/Plugins/Animation/AnimToTexture/Source/AnimToTexture/Private/AnimToTextureDataAsset.cpp
+++ b/Engine/Plugins/Animation/AnimToTexture/Source/AnimToTexture/Private/AnimToTextureDataAsset.cpp
@@ -22,30 +22,20 @@ int32 UAnimToTextureDataAsset::GetIndexFromAnimSequence(const UAnimSequence* Seq
 	return OutIndex;
 }
 
-namespace AnimToTextureParamNames
+void UAnimToTextureDataAsset::ResetInfo()
 {
-	static const FName BoundingBoxMin = TEXT("MinBBox");
-	static const FName BoundingBoxScale = TEXT("SizeBBox");
-	static const FName RowsPerFrame = TEXT("RowsPerFrame");
-	static const FName BoneWeightRowsPerFrame = TEXT("BoneWeightsRowsPerFrame");
-	static const FName NumFrames = TEXT("NumFrames (S)");
-	static const FName VertexPositionTexture = TEXT("PositionTexture");
-	static const FName VertexNormalTexture = TEXT("NormalTexture");
-	static const FName BonePositionTexture = TEXT("BonePositionTexture");
-	static const FName BoneRotationTexture = TEXT("BoneRotationTexture");
-	static const FName BoneWeightsTexture = TEXT("BoneWeightsTexture");
-}
+	// Common Info.
+	NumFrames = 0;
+	Animations.Reset();
 
-FAnimToTextureMaterialParamNames::FAnimToTextureMaterialParamNames()
-{
-	BoundingBoxMin = AnimToTextureParamNames::BoundingBoxMin;
-	BoundingBoxScale = AnimToTextureParamNames::BoundingBoxScale;
-	RowsPerFrame = AnimToTextureParamNames::RowsPerFrame;
-	BoneWeightRowsPerFrame = AnimToTextureParamNames::BoneWeightRowsPerFrame;
-	NumFrames = AnimToTextureParamNames::NumFrames;
-	VertexPositionTexture = AnimToTextureParamNames::VertexPositionTexture;
-	VertexNormalTexture = AnimToTextureParamNames::VertexNormalTexture;
-	BonePositionTexture = AnimToTextureParamNames::BonePositionTexture;
-	BoneRotationTexture = AnimToTextureParamNames::BoneRotationTexture;
-	BoneWeightsTexture = AnimToTextureParamNames::BoneWeightsTexture;
-}
+	// Vertex Info
+	VertexRowsPerFrame = 1;
+	VertexMinBBox = FVector::ZeroVector;
+	VertexSizeBBox = FVector::ZeroVector;
+
+	// Bone Info
+	BoneRowsPerFrame = 1;
+	BoneWeightRowsPerFrame = 1;
+	BoneMinBBox = FVector::ZeroVector;
+	BoneSizeBBox = FVector::ZeroVector;
+};
\ No newline at end of file
diff --git a/Engine/Plugins/Animation/AnimToTexture/Source/AnimToTexture/Public/AnimToTextureDataAsset.h b/Engine/Plugins/Animation/AnimToTexture/Source/AnimToTexture/Public/AnimToTextureDataAsset.h
--- a/Engine/Plugins/Animation/AnimToTexture/Source/AnimToTexture/Public/AnimToTextureDataAsset.h
+++ b/Engine/Plugins/Animation/AnimToTexture/Source/AnimToTexture/Public/AnimToTextureDataAsset.h
@@ -13,56 +13,26 @@ class USkeletalMesh;
 class UStaticMesh;
 class UTexture2D;
 
-USTRUCT(Blueprintable)
-struct ANIMTOTEXTURE_API FAnimToTextureMaterialParamNames
+namespace AnimToTextureParamNames
 {
-	GENERATED_BODY()
-
-	//
-	// Scalar Parameters
-	//
-
-	UPROPERTY(BlueprintReadWrite, Category = Default, EditAnywhere)
-	FName RowsPerFrame;
-
-	UPROPERTY(BlueprintReadWrite, Category = Default, EditAnywhere)
-	FName BoneWeightRowsPerFrame;
-
-	UPROPERTY(BlueprintReadWrite, Category = Default, EditAnywhere)
-	FName NumFrames;
-
-	//
-	// Vector Parameters
-	//
-
-	UPROPERTY(BlueprintReadWrite, Category = Default, EditAnywhere)
-	FName BoundingBoxMin;
-
-	UPROPERTY(BlueprintReadWrite, Category = Default, EditAnywhere)
-	FName BoundingBoxScale;
-
-	//
-	// Texture Parameters
-	//
-
-	UPROPERTY(BlueprintReadWrite, Category = Default, EditAnywhere)
-	FName VertexPositionTexture;
-
-	UPROPERTY(BlueprintReadWrite, Category = Default, EditAnywhere)
-	FName VertexNormalTexture;
-
-	UPROPERTY(BlueprintReadWrite, Category = Default, EditAnywhere)
-	FName BonePositionTexture;
-
-	UPROPERTY(BlueprintReadWrite, Category = Default, EditAnywhere)
-	FName BoneRotationTexture;
-
-	UPROPERTY(BlueprintReadWrite, Category = Default, EditAnywhere)
-	FName BoneWeightsTexture;
-
-	// Initialize Names
-	FAnimToTextureMaterialParamNames();
-};
+	static const FName Animate = TEXT("Animate");
+	static const FName BoundingBoxMin = TEXT("MinBBox");
+	static const FName BoundingBoxScale = TEXT("SizeBBox");
+	static const FName RowsPerFrame = TEXT("RowsPerFrame");
+	static const FName BoneWeightRowsPerFrame = TEXT("BoneWeightsRowsPerFrame");
+	static const FName NumFrames = TEXT("NumFrames");
+	static const FName VertexPositionTexture = TEXT("PositionTexture");
+	static const FName VertexNormalTexture = TEXT("NormalTexture");
+	static const FName BonePositionTexture = TEXT("BonePositionTexture");
+	static const FName BoneRotationTexture = TEXT("BoneRotationTexture");
+	static const FName BoneWeightsTexture = TEXT("BoneWeightsTexture");
+	static const FName UseUV0 = TEXT("UseUV0");
+	static const FName UseUV1 = TEXT("UseUV1");
+	static const FName UseUV2 = TEXT("UseUV2");
+	static const FName UseUV3 = TEXT("UseUV3");
+	static const FName UseTwoInfluences = TEXT("UseTwoInfluences");
+	static const FName UseFourInfluences = TEXT("UseFourInfluences");
+}
 
 UENUM(Blueprintable)
 enum class EAnimToTextureMode : uint8
@@ -74,14 +44,25 @@ enum class EAnimToTextureMode : uint8
 };
 
 UENUM(Blueprintable)
-enum class EAnimToTextureBonePrecision : uint8
+enum class EAnimToTexturePrecision : uint8
 {
-	/* Bone positions and rotations stored in 8 bits */
+	/* 8 bits */
 	EightBits,
-	/* Bone positions and rotations stored in 16 bits */
+	/* 16 bits */
 	SixteenBits,
 };
 
+UENUM(Blueprintable)
+enum class EAnimToTextureNumBoneInfluences : uint8
+{
+	/* Single bone influence */
+	One,
+	/* Blend between two influences */
+	Two,
+	/* Blend between four influences */
+	Four,
+};
+
 USTRUCT(Blueprintable)
 struct FAnimSequenceInfo
 {
@@ -90,19 +71,19 @@ struct FAnimSequenceInfo
 	UPROPERTY(EditAnywhere, Category = Default, BlueprintReadWrite)
 	bool bEnabled = true;
 
-	UPROPERTY(EditAnywhere, Category = Default, BlueprintReadWrite)
+	UPROPERTY(EditAnywhere, Category = Default, BlueprintReadWrite, meta = (EditCondition = "bEnabled", EditConditionHides))
 	TObjectPtr<UAnimSequence> AnimSequence = nullptr;
 
-	UPROPERTY(EditAnywhere, Category = Default, BlueprintReadWrite)
+	UPROPERTY(EditAnywhere, Category = Default, BlueprintReadWrite, meta = (EditCondition = "bEnabled", EditConditionHides))
 	bool bLooping = true;
 
-	UPROPERTY(EditAnywhere, Category = Default, BlueprintReadWrite)
+	UPROPERTY(EditAnywhere, Category = Default, BlueprintReadWrite, meta = (EditCondition = "bEnabled", EditConditionHides))
 	bool bUseCustomRange = false;
 
-	UPROPERTY(EditAnywhere, Category = Default, BlueprintReadWrite, meta = (EditCondition = "bUseCustomRange"))
+	UPROPERTY(EditAnywhere, Category = Default, BlueprintReadWrite, meta = (EditCondition = "bEnabled && bUseCustomRange", EditConditionHides))
 	int32 StartFrame = 0;
 
-	UPROPERTY(EditAnywhere, Category = Default, BlueprintReadWrite, meta = (EditCondition = "bUseCustomRange"))
+	UPROPERTY(EditAnywhere, Category = Default, BlueprintReadWrite, meta = (EditCondition = "bEnabled && bUseCustomRange", EditConditionHides))
 	int32 EndFrame = 1;
 
 };
@@ -182,6 +163,13 @@ public:
 	UPROPERTY(EditAnywhere, BlueprintReadWrite, Category = "Texture")
 	bool bEnforcePowerOfTwo = false;
 
+	/**
+	* Texture Precision
+	* 16 bits is only supported for power-of-2 resolutions.
+	*/
+	UPROPERTY(EditAnywhere, BlueprintReadWrite, Category = "Texture", meta = (EditCondition="bEnforcePowerOfTwo"))
+	EAnimToTexturePrecision Precision = EAnimToTexturePrecision::EightBits;
+
 	/**
 	* Storage Mode.
 	* Vertex: will store per-vertex position and normal.
@@ -195,41 +183,35 @@ public:
 	* Texture for storing vertex positions
 	* This is only used on Vertex Mode
 	*/
-	UPROPERTY(EditAnywhere, BlueprintReadWrite, Category = "Texture|Vertex", meta = (EditCondition = "Mode == EAnimToTextureMode::Vertex"))
+	UPROPERTY(EditAnywhere, BlueprintReadWrite, Category = "Texture", meta = (EditCondition = "Mode == EAnimToTextureMode::Vertex", EditConditionHides))
 	TSoftObjectPtr<UTexture2D> VertexPositionTexture;
 
 	/**
 	* Texture for storing vertex normals
 	* This is only used on Vertex Mode
 	*/
-	UPROPERTY(EditAnywhere, BlueprintReadWrite, Category = "Texture|Vertex", meta = (EditCondition = "Mode == EAnimToTextureMode::Vertex"))
+	UPROPERTY(EditAnywhere, BlueprintReadWrite, Category = "Texture", meta = (EditCondition = "Mode == EAnimToTextureMode::Vertex", EditConditionHides))
 	TSoftObjectPtr<UTexture2D> VertexNormalTexture;
 
-	/**
-	* Texture Precision for: BonePosition and BoneRotations.
-	*/
-	UPROPERTY(EditAnywhere, BlueprintReadWrite, Category = "Texture|Bone", meta = (EditCondition = "Mode == EAnimToTextureMode::Bone"))
-	EAnimToTextureBonePrecision BonePrecision = EAnimToTextureBonePrecision::EightBits;
-
 	/**
 	* Texture for storing bone positions
 	* This is only used on Bone Mode
 	*/
-	UPROPERTY(EditAnywhere, BlueprintReadWrite, Category = "Texture|Bone", meta = (EditCondition = "Mode == EAnimToTextureMode::Bone"))
+	UPROPERTY(EditAnywhere, BlueprintReadWrite, Category = "Texture", meta = (EditCondition = "Mode == EAnimToTextureMode::Bone", EditConditionHides))
 	TSoftObjectPtr<UTexture2D> BonePositionTexture;
 
 	/**
 	* Texture for storing bone rotations
 	* This is only used on Bone Mode
 	*/
-	UPROPERTY(EditAnywhere, BlueprintReadWrite, Category = "Texture|Bone", meta = (EditCondition = "Mode == EAnimToTextureMode::Bone"))
+	UPROPERTY(EditAnywhere, BlueprintReadWrite, Category = "Texture", meta = (EditCondition = "Mode == EAnimToTextureMode::Bone", EditConditionHides))
 	TSoftObjectPtr<UTexture2D> BoneRotationTexture;
 
 	/**
 	* Texture for storing vertex/bone weighting
 	* This is only used on Bone Mode
 	*/
-	UPROPERTY(EditAnywhere, BlueprintReadWrite, Category = "Texture|Bone", meta = (EditCondition = "Mode == EAnimToTextureMode::Bone"))
+	UPROPERTY(EditAnywhere, BlueprintReadWrite, Category = "Texture", meta = (EditCondition = "Mode == EAnimToTextureMode::Bone", EditConditionHides))
 	TSoftObjectPtr<UTexture2D> BoneWeightTexture;
 
 	// ------------------------------------------------------
@@ -255,25 +237,25 @@ public:
 	UPROPERTY(VisibleAnywhere, BlueprintReadOnly, Category = "Info")
 	int32 NumFrames = 0;
 
-	UPROPERTY(VisibleAnywhere, BlueprintReadOnly, Category = "Info|Vertex", Meta = (DisplayName = "RowsPerFrame"))
+	UPROPERTY(VisibleAnywhere, BlueprintReadOnly, Category = "Info", Meta = (EditCondition = "Mode == EAnimToTextureMode::Vertex", EditConditionHides))
 	int32 VertexRowsPerFrame = 1;
 
-	UPROPERTY(VisibleAnywhere, BlueprintReadOnly, Category = "Info|Vertex", Meta = (DisplayName = "MinBBox"))
+	UPROPERTY(VisibleAnywhere, BlueprintReadOnly, Category = "Info", Meta = (DisplayName = "MinBBox", EditCondition = "Mode == EAnimToTextureMode::Vertex", EditConditionHides))
 	FVector VertexMinBBox;
 
-	UPROPERTY(VisibleAnywhere, BlueprintReadOnly, Category = "Info|Vertex", Meta = (DisplayName = "SizeBBox"))
+	UPROPERTY(VisibleAnywhere, BlueprintReadOnly, Category = "Info", Meta = (DisplayName = "SizeBBox", EditCondition = "Mode == EAnimToTextureMode::Vertex", EditConditionHides))
 	FVector VertexSizeBBox;
 
-	UPROPERTY(VisibleAnywhere, BlueprintReadOnly, Category = "Info|Bone")
+	UPROPERTY(VisibleAnywhere, BlueprintReadOnly, Category = "Info", Meta = (EditCondition = "Mode == EAnimToTextureMode::Bone", EditConditionHides))
 	int32 BoneWeightRowsPerFrame = 1;
 
-	UPROPERTY(VisibleAnywhere, BlueprintReadOnly, Category = "Info|Bone")
+	UPROPERTY(VisibleAnywhere, BlueprintReadOnly, Category = "Info", Meta = (EditCondition = "Mode == EAnimToTextureMode::Bone", EditConditionHides))
 	int32 BoneRowsPerFrame = 1;
 
-	UPROPERTY(VisibleAnywhere, BlueprintReadOnly, Category = "Info|Bone", Meta = (DisplayName = "MinBBox"))
+	UPROPERTY(VisibleAnywhere, BlueprintReadOnly, Category = "Info", Meta = (DisplayName = "MinBBox", EditCondition = "Mode == EAnimToTextureMode::Bone", EditConditionHides))
 	FVector BoneMinBBox;
 
-	UPROPERTY(VisibleAnywhere, BlueprintReadOnly, Category = "Info|Bone", Meta = (DisplayName = "SizeBBox"))
+	UPROPERTY(VisibleAnywhere, BlueprintReadOnly, Category = "Info", Meta = (DisplayName = "SizeBBox", EditCondition = "Mode == EAnimToTextureMode::Bone", EditConditionHides))
 	FVector BoneSizeBBox;
 
 	UPROPERTY(VisibleAnywhere, BlueprintReadOnly, Category = "Info")
@@ -283,7 +265,7 @@ public:
 	int32 GetIndexFromAnimSequence(const UAnimSequence* Sequence);
 
 	UFUNCTION()
-	void Reset();
+	void ResetInfo();
 
 	// If we weren't in a plugin, we could unify this in a base class
 	template<typename AssetType>
@@ -331,21 +313,3 @@ public:
 	UFUNCTION(BlueprintPure, Category = Default, meta = (DisplayName = "Get Bone Weight Texture"))
 	UTexture2D* BP_GetBoneWeightTexture() { return GetBoneWeightTexture(); }
 };
-
-FORCEINLINE void UAnimToTextureDataAsset::Reset()
-{
-	// Common Info.
-	this->NumFrames = 0;
-	this->Animations.Reset();
-
-	// Vertex Info
-	this->VertexRowsPerFrame = 1;
-	this->VertexMinBBox = FVector::ZeroVector;
-	this->VertexSizeBBox = FVector::ZeroVector;
-
-	// Bone Info
-	this->BoneRowsPerFrame = 1;
-	this->BoneWeightRowsPerFrame = 1;
-	this->BoneMinBBox = FVector::ZeroVector; 
-	this->BoneSizeBBox = FVector::ZeroVector;
-};
\ No newline at end of file
diff --git a/Engine/Plugins/Animation/AnimToTexture/Source/AnimToTexture/Public/AnimToTextureInstancePlaybackHelpers.h b/Engine/Plugins/Animation/AnimToTexture/Source/AnimToTexture/Public/AnimToTextureInstancePlaybackHelpers.h
--- a/Engine/Plugins/Animation/AnimToTexture/Source/AnimToTexture/Public/AnimToTextureInstancePlaybackHelpers.h
+++ b/Engine/Plugins/Animation/AnimToTexture/Source/AnimToTexture/Public/AnimToTextureInstancePlaybackHelpers.h
@@ -75,24 +75,24 @@ class ANIMTOTEXTURE_API UAnimToTextureInstancePlaybackLibrary : public UBlueprin
 
 public:
 
-	UFUNCTION(BlueprintCallable, Category = "Anim Texture Playback")
+	UFUNCTION(BlueprintCallable, Category = "AnimToTexture|Playback")
 	static void SetupInstancedMeshComponent(UInstancedStaticMeshComponent* InstancedMeshComponent, UPARAM(ref) FAnimToTextureInstanceData& InstanceData, int32 NumInstances);
 
-	UFUNCTION(BlueprintCallable, Category = "Anim Texture Playback")
+	UFUNCTION(BlueprintCallable, Category = "AnimToTexture|Playback")
 	static void BatchUpdateInstancedMeshComponent(UInstancedStaticMeshComponent* InstancedMeshComponent, UPARAM(ref) FAnimToTextureInstanceData& InstanceData);
 
-	UFUNCTION(BlueprintCallable, Category = "Anim Texture Playback")
+	UFUNCTION(BlueprintCallable, Category = "AnimToTexture|Playback")
 	static void AllocateInstanceData(UPARAM(ref) FAnimToTextureInstanceData& InstanceData, int32 Count);
 
-	UFUNCTION(BlueprintCallable, Category = "Anim Texture Playback")
+	UFUNCTION(BlueprintCallable, Category = "AnimToTexture|Playback")
 	static bool UpdateInstanceData(UPARAM(ref) FAnimToTextureInstanceData& InstanceData, int32 InstanceIndex, const FAnimToTextureInstancePlaybackData& PlaybackData, const FTransform& Transform);
 
-	UFUNCTION(BlueprintCallable, Category = "Anim Texture Playback")
+	UFUNCTION(BlueprintCallable, Category = "AnimToTexture|Playback")
 	static bool GetInstancePlaybackData(UPARAM(ref) const FAnimToTextureInstanceData& InstanceData, int32 InstanceIndex, FAnimToTextureInstancePlaybackData& InstancePlaybackData);
 
-	UFUNCTION(BlueprintCallable, Category = "Anim Texture Playback")
+	UFUNCTION(BlueprintCallable, Category = "AnimToTexture|Playback")
 	static bool GetInstanceTransform(UPARAM(ref) const FAnimToTextureInstanceData& InstanceData, int32 InstanceIndex, FTransform& InstanceTransform);
 
-	UFUNCTION(BlueprintPure, BlueprintCallable, Category = "Anim Texture Playback")
+	UFUNCTION(BlueprintPure, BlueprintCallable, Category = "AnimToTexture|Playback")
 	static bool AnimStateFromDataAsset(const UAnimToTextureDataAsset* DataAsset, int32 StateIndex, FAnimToTextureAnimState& AnimState);
 };
diff --git a/Engine/Plugins/Animation/AnimToTexture/Source/AnimToTextureEditor/Private/AnimToTextureBPLibrary.cpp b/Engine/Plugins/Animation/AnimToTexture/Source/AnimToTextureEditor/Private/AnimToTextureBPLibrary.cpp
--- a/Engine/Plugins/Animation/AnimToTexture/Source/AnimToTextureEditor/Private/AnimToTextureBPLibrary.cpp
+++ b/Engine/Plugins/Animation/AnimToTexture/Source/AnimToTextureEditor/Private/AnimToTextureBPLibrary.cpp
@@ -40,8 +40,8 @@ void UAnimToTextureBPLibrary::AnimationToTexture(UAnimToTextureDataAsset* DataAs
 		return;
 	}
 
-	// Reset DataAsset Values
-	DataAsset->Reset();
+	// Reset DataAsset Info Values
+	DataAsset->ResetInfo();
 
 	if (!DataAsset->GetSkeletalMesh() || !DataAsset->GetStaticMesh())
 	{
@@ -356,12 +356,22 @@ void UAnimToTextureBPLibrary::AnimationToTexture(UAnimToTextureDataAsset* DataAs
 			NormalizedVertexDeltas, NormalizedVertexNormals);
 
 		// Write Textures
-		AnimToTexture_Private::WriteVectorsToTexture<FVector3f, AnimToTexture_Private::FLowPrecision>(NormalizedVertexDeltas, DataAsset->NumFrames, DataAsset->VertexRowsPerFrame, Height, Width, DataAsset->GetVertexPositionTexture());
-		AnimToTexture_Private::WriteVectorsToTexture<FVector3f, AnimToTexture_Private::FLowPrecision>(NormalizedVertexNormals, DataAsset->NumFrames, DataAsset->VertexRowsPerFrame, Height, Width, DataAsset->GetVertexNormalTexture());
-
-	
+		if (DataAsset->Precision == EAnimToTexturePrecision::SixteenBits && DataAsset->bEnforcePowerOfTwo)
+		{
+			AnimToTexture_Private::WriteVectorsToTexture<FVector3f, AnimToTexture_Private::FHighPrecision>(NormalizedVertexDeltas, DataAsset->NumFrames, DataAsset->VertexRowsPerFrame, Height, Width, DataAsset->GetVertexPositionTexture());
+			AnimToTexture_Private::WriteVectorsToTexture<FVector3f, AnimToTexture_Private::FHighPrecision>(NormalizedVertexNormals, DataAsset->NumFrames, DataAsset->VertexRowsPerFrame, Height, Width, DataAsset->GetVertexNormalTexture());
+		}
+		else
+		{
+			AnimToTexture_Private::WriteVectorsToTexture<FVector3f, AnimToTexture_Private::FLowPrecision>(NormalizedVertexDeltas, DataAsset->NumFrames, DataAsset->VertexRowsPerFrame, Height, Width, DataAsset->GetVertexPositionTexture());
+			AnimToTexture_Private::WriteVectorsToTexture<FVector3f, AnimToTexture_Private::FLowPrecision>(NormalizedVertexNormals, DataAsset->NumFrames, DataAsset->VertexRowsPerFrame, Height, Width, DataAsset->GetVertexNormalTexture());
+		}
+			
 		// Add Vertex UVChannel
 		CreateUVChannel(DataAsset->GetStaticMesh(), DataAsset->StaticLODIndex, DataAsset->UVChannel, Height, Width);
+		
+		// Update Bounds
+		SetBoundsExtensions(DataAsset->GetStaticMesh(), DataAsset->VertexMinBBox, DataAsset->VertexSizeBBox);
 	}
 
 	// ---------------------------------------------------------------------------
@@ -393,7 +403,7 @@ void UAnimToTextureBPLibrary::AnimationToTexture(UAnimToTextureDataAsset* DataAs
 				NormalizedBonePositions, NormalizedBoneRotations);
 
 			// Write Textures
-			if (DataAsset->BonePrecision == EAnimToTextureBonePrecision::SixteenBits)
+			if (DataAsset->Precision == EAnimToTexturePrecision::SixteenBits && DataAsset->bEnforcePowerOfTwo)
 			{
 				AnimToTexture_Private::WriteVectorsToTexture<FVector3f, AnimToTexture_Private::FHighPrecision>(NormalizedBonePositions, DataAsset->NumFrames + 1, DataAsset->BoneRowsPerFrame, Height, Width, DataAsset->GetBonePositionTexture());
 				AnimToTexture_Private::WriteVectorsToTexture<FVector4, AnimToTexture_Private::FHighPrecision>(NormalizedBoneRotations, DataAsset->NumFrames + 1, DataAsset->BoneRowsPerFrame, Height, Width, DataAsset->GetBoneRotationTexture());
@@ -448,6 +458,9 @@ void UAnimToTextureBPLibrary::AnimationToTexture(UAnimToTextureDataAsset* DataAs
 
 		// Add Vertex UVChannel
 		CreateUVChannel(DataAsset->GetStaticMesh(), DataAsset->StaticLODIndex, DataAsset->UVChannel, Height, Width);
+
+		// Update Bounds
+		SetBoundsExtensions(DataAsset->GetStaticMesh(), DataAsset->BoneMinBBox, DataAsset->BoneSizeBBox);
 	}
 
 	// ---------------------------------------------------------------------------
@@ -458,102 +471,117 @@ void UAnimToTextureBPLibrary::AnimationToTexture(UAnimToTextureDataAsset* DataAs
 
 
 void UAnimToTextureBPLibrary::UpdateMaterialInstanceFromDataAsset(UAnimToTextureDataAsset* DataAsset, UMaterialInstanceConstant* MaterialInstance, 
-	const FAnimToTextureMaterialParamNames& ParamNames, const EMaterialParameterAssociation MaterialParameterAssociation)
+	const bool bAnimate, const EAnimToTextureNumBoneInfluences NumBoneInfluences, const EMaterialParameterAssociation MaterialParameterAssociation)
 {
 	if (!MaterialInstance || !DataAsset)
 	{
 		return;
 	}
-	
-	// Set Preview Mesh
-	if (DataAsset->GetStaticMesh())
-	{
-		MaterialInstance->PreviewMesh = DataAsset->GetStaticMesh();
-	}
 
+	// Set UVChannel
+	switch (DataAsset->UVChannel)
+	{
+		case 0:
+			UMaterialEditingLibrary::SetMaterialInstanceStaticSwitchParameterValue(MaterialInstance, AnimToTextureParamNames::UseUV0, true, MaterialParameterAssociation);
+			UMaterialEditingLibrary::SetMaterialInstanceStaticSwitchParameterValue(MaterialInstance, AnimToTextureParamNames::UseUV1, false, MaterialParameterAssociation);
+			UMaterialEditingLibrary::SetMaterialInstanceStaticSwitchParameterValue(MaterialInstance, AnimToTextureParamNames::UseUV2, false, MaterialParameterAssociation);
+			UMaterialEditingLibrary::SetMaterialInstanceStaticSwitchParameterValue(MaterialInstance, AnimToTextureParamNames::UseUV3, false, MaterialParameterAssociation);
+			break;
+		case 1:
+			UMaterialEditingLibrary::SetMaterialInstanceStaticSwitchParameterValue(MaterialInstance, AnimToTextureParamNames::UseUV0, false, MaterialParameterAssociation);
+			UMaterialEditingLibrary::SetMaterialInstanceStaticSwitchParameterValue(MaterialInstance, AnimToTextureParamNames::UseUV1, true, MaterialParameterAssociation);
+			UMaterialEditingLibrary::SetMaterialInstanceStaticSwitchParameterValue(MaterialInstance, AnimToTextureParamNames::UseUV2, false, MaterialParameterAssociation);
+			UMaterialEditingLibrary::SetMaterialInstanceStaticSwitchParameterValue(MaterialInstance, AnimToTextureParamNames::UseUV3, false, MaterialParameterAssociation);
+			break;
+		case 2:
+			UMaterialEditingLibrary::SetMaterialInstanceStaticSwitchParameterValue(MaterialInstance, AnimToTextureParamNames::UseUV0, false, MaterialParameterAssociation);
+			UMaterialEditingLibrary::SetMaterialInstanceStaticSwitchParameterValue(MaterialInstance, AnimToTextureParamNames::UseUV1, false, MaterialParameterAssociation);
+			UMaterialEditingLibrary::SetMaterialInstanceStaticSwitchParameterValue(MaterialInstance, AnimToTextureParamNames::UseUV2, true, MaterialParameterAssociation);
+			UMaterialEditingLibrary::SetMaterialInstanceStaticSwitchParameterValue(MaterialInstance, AnimToTextureParamNames::UseUV3, false, MaterialParameterAssociation);
+			break;
+		case 3:
+			UMaterialEditingLibrary::SetMaterialInstanceStaticSwitchParameterValue(MaterialInstance, AnimToTextureParamNames::UseUV0, false, MaterialParameterAssociation);
+			UMaterialEditingLibrary::SetMaterialInstanceStaticSwitchParameterValue(MaterialInstance, AnimToTextureParamNames::UseUV1, false, MaterialParameterAssociation);
+			UMaterialEditingLibrary::SetMaterialInstanceStaticSwitchParameterValue(MaterialInstance, AnimToTextureParamNames::UseUV2, false, MaterialParameterAssociation);
+			UMaterialEditingLibrary::SetMaterialInstanceStaticSwitchParameterValue(MaterialInstance, AnimToTextureParamNames::UseUV3, true, MaterialParameterAssociation);
+			break;
+		default:
+			UMaterialEditingLibrary::SetMaterialInstanceStaticSwitchParameterValue(MaterialInstance, AnimToTextureParamNames::UseUV0, false, MaterialParameterAssociation);
+			UMaterialEditingLibrary::SetMaterialInstanceStaticSwitchParameterValue(MaterialInstance, AnimToTextureParamNames::UseUV1, true, MaterialParameterAssociation);
+			UMaterialEditingLibrary::SetMaterialInstanceStaticSwitchParameterValue(MaterialInstance, AnimToTextureParamNames::UseUV2, false, MaterialParameterAssociation);
+			UMaterialEditingLibrary::SetMaterialInstanceStaticSwitchParameterValue(MaterialInstance, AnimToTextureParamNames::UseUV3, false, MaterialParameterAssociation);
+			break;
+	}
+
+	// Update Vertex Params
 	if (DataAsset->Mode == EAnimToTextureMode::Vertex)
 	{
 		FLinearColor VectorParameter;
 		VectorParameter = FLinearColor(DataAsset->VertexMinBBox);
-		UMaterialEditingLibrary::SetMaterialInstanceVectorParameterValue(MaterialInstance, ParamNames.BoundingBoxMin, VectorParameter, MaterialParameterAssociation);
+		UMaterialEditingLibrary::SetMaterialInstanceVectorParameterValue(MaterialInstance, AnimToTextureParamNames::BoundingBoxMin, VectorParameter, MaterialParameterAssociation);
 		
 		VectorParameter = FLinearColor(DataAsset->VertexSizeBBox);
-		UMaterialEditingLibrary::SetMaterialInstanceVectorParameterValue(MaterialInstance, ParamNames.BoundingBoxScale, VectorParameter, MaterialParameterAssociation);
-		UMaterialEditingLibrary::SetMaterialInstanceScalarParameterValue(MaterialInstance, ParamNames.NumFrames, DataAsset->NumFrames, MaterialParameterAssociation);
-		UMaterialEditingLibrary::SetMaterialInstanceScalarParameterValue(MaterialInstance, ParamNames.RowsPerFrame, DataAsset->VertexRowsPerFrame, MaterialParameterAssociation);
+		UMaterialEditingLibrary::SetMaterialInstanceVectorParameterValue(MaterialInstance, AnimToTextureParamNames::BoundingBoxScale, VectorParameter, MaterialParameterAssociation);
+		UMaterialEditingLibrary::SetMaterialInstanceScalarParameterValue(MaterialInstance, AnimToTextureParamNames::NumFrames, DataAsset->NumFrames, MaterialParameterAssociation);
+		UMaterialEditingLibrary::SetMaterialInstanceScalarParameterValue(MaterialInstance, AnimToTextureParamNames::RowsPerFrame, DataAsset->VertexRowsPerFrame, MaterialParameterAssociation);
 
-		UMaterialEditingLibrary::SetMaterialInstanceTextureParameterValue(MaterialInstance, ParamNames.VertexPositionTexture, DataAsset->GetVertexPositionTexture(), MaterialParameterAssociation);
-		UMaterialEditingLibrary::SetMaterialInstanceTextureParameterValue(MaterialInstance, ParamNames.VertexNormalTexture, DataAsset->GetVertexNormalTexture(), MaterialParameterAssociation);
+		UMaterialEditingLibrary::SetMaterialInstanceTextureParameterValue(MaterialInstance, AnimToTextureParamNames::VertexPositionTexture, DataAsset->GetVertexPositionTexture(), MaterialParameterAssociation);
+		UMaterialEditingLibrary::SetMaterialInstanceTextureParameterValue(MaterialInstance, AnimToTextureParamNames::VertexNormalTexture, DataAsset->GetVertexNormalTexture(), MaterialParameterAssociation);
 
 	}
+
+	// Update Bone Params
 	else if (DataAsset->Mode == EAnimToTextureMode::Bone)
 	{
 		FLinearColor VectorParameter;
 		VectorParameter = FLinearColor(DataAsset->BoneMinBBox);
-		UMaterialEditingLibrary::SetMaterialInstanceVectorParameterValue(MaterialInstance, ParamNames.BoundingBoxMin, VectorParameter, MaterialParameterAssociation);
+		UMaterialEditingLibrary::SetMaterialInstanceVectorParameterValue(MaterialInstance, AnimToTextureParamNames::BoundingBoxMin, VectorParameter, MaterialParameterAssociation);
 
 		VectorParameter = FLinearColor(DataAsset->BoneSizeBBox);
-		UMaterialEditingLibrary::SetMaterialInstanceVectorParameterValue(MaterialInstance, ParamNames.BoundingBoxScale, VectorParameter, MaterialParameterAssociation);
-		UMaterialEditingLibrary::SetMaterialInstanceScalarParameterValue(MaterialInstance, ParamNames.NumFrames, DataAsset->NumFrames, MaterialParameterAssociation);
-		UMaterialEditingLibrary::SetMaterialInstanceScalarParameterValue(MaterialInstance, ParamNames.RowsPerFrame, DataAsset->BoneRowsPerFrame, MaterialParameterAssociation);
-		UMaterialEditingLibrary::SetMaterialInstanceScalarParameterValue(MaterialInstance, ParamNames.BoneWeightRowsPerFrame, DataAsset->BoneWeightRowsPerFrame, MaterialParameterAssociation);
+		UMaterialEditingLibrary::SetMaterialInstanceVectorParameterValue(MaterialInstance, AnimToTextureParamNames::BoundingBoxScale, VectorParameter, MaterialParameterAssociation);
+		UMaterialEditingLibrary::SetMaterialInstanceScalarParameterValue(MaterialInstance, AnimToTextureParamNames::NumFrames, DataAsset->NumFrames, MaterialParameterAssociation);
+		UMaterialEditingLibrary::SetMaterialInstanceScalarParameterValue(MaterialInstance, AnimToTextureParamNames::RowsPerFrame, DataAsset->BoneRowsPerFrame, MaterialParameterAssociation);
+		UMaterialEditingLibrary::SetMaterialInstanceScalarParameterValue(MaterialInstance, AnimToTextureParamNames::BoneWeightRowsPerFrame, DataAsset->BoneWeightRowsPerFrame, MaterialParameterAssociation);
+
+		UMaterialEditingLibrary::SetMaterialInstanceTextureParameterValue(MaterialInstance, AnimToTextureParamNames::BonePositionTexture, DataAsset->GetBonePositionTexture(), MaterialParameterAssociation);
+		UMaterialEditingLibrary::SetMaterialInstanceTextureParameterValue(MaterialInstance, AnimToTextureParamNames::BoneRotationTexture, DataAsset->GetBoneRotationTexture(), MaterialParameterAssociation);
+		UMaterialEditingLibrary::SetMaterialInstanceTextureParameterValue(MaterialInstance, AnimToTextureParamNames::BoneWeightsTexture, DataAsset->GetBoneWeightTexture(), MaterialParameterAssociation);
+
+		// Num Influences
+		switch (NumBoneInfluences)
+		{
+			case EAnimToTextureNumBoneInfluences::One:
+				UMaterialEditingLibrary::SetMaterialInstanceStaticSwitchParameterValue(MaterialInstance, AnimToTextureParamNames::UseTwoInfluences, false, MaterialParameterAssociation);
+				UMaterialEditingLibrary::SetMaterialInstanceStaticSwitchParameterValue(MaterialInstance, AnimToTextureParamNames::UseFourInfluences, false, MaterialParameterAssociation);
+				break;
+			case EAnimToTextureNumBoneInfluences::Two:
+				UMaterialEditingLibrary::SetMaterialInstanceStaticSwitchParameterValue(MaterialInstance, AnimToTextureParamNames::UseTwoInfluences, true, MaterialParameterAssociation);
+				UMaterialEditingLibrary::SetMaterialInstanceStaticSwitchParameterValue(MaterialInstance, AnimToTextureParamNames::UseFourInfluences, false, MaterialParameterAssociation);
+				break;
+			case EAnimToTextureNumBoneInfluences::Four:
+				UMaterialEditingLibrary::SetMaterialInstanceStaticSwitchParameterValue(MaterialInstance, AnimToTextureParamNames::UseTwoInfluences, false, MaterialParameterAssociation);
+				UMaterialEditingLibrary::SetMaterialInstanceStaticSwitchParameterValue(MaterialInstance, AnimToTextureParamNames::UseFourInfluences, true, MaterialParameterAssociation);
+				break;
+		}
 
-		UMaterialEditingLibrary::SetMaterialInstanceTextureParameterValue(MaterialInstance, ParamNames.BonePositionTexture, DataAsset->GetBonePositionTexture(), MaterialParameterAssociation);
-		UMaterialEditingLibrary::SetMaterialInstanceTextureParameterValue(MaterialInstance, ParamNames.BoneRotationTexture, DataAsset->GetBoneRotationTexture(), MaterialParameterAssociation);
-		UMaterialEditingLibrary::SetMaterialInstanceTextureParameterValue(MaterialInstance, ParamNames.BoneWeightsTexture, DataAsset->GetBoneWeightTexture(), MaterialParameterAssociation);
 	}
 
+	// Animate
+	UMaterialEditingLibrary::SetMaterialInstanceStaticSwitchParameterValue(MaterialInstance, AnimToTextureParamNames::Animate, bAnimate, MaterialParameterAssociation);
+
 	// Update Material
 	UMaterialEditingLibrary::UpdateMaterialInstance(MaterialInstance);
 
 	// Rebuild Material
 	UMaterialEditingLibrary::RebuildMaterialInstanceEditors(MaterialInstance->GetMaterial());
-}
-
-void UAnimToTextureBPLibrary::UpdateMaterialLayerFunction(
-	UMaterialInstanceConstant* MaterialInstance,
-	UMaterialFunctionInterface* OldMaterialFunction,
-	UMaterialFunctionInterface* NewMaterialFunction)
-{
-	if (!IsValid(MaterialInstance))
-	{
-		return;
-	}
-
-	const FMaterialParameterInfo ParameterInfo(TEXT("BoneAnimation"), EMaterialParameterAssociation::GlobalParameter);
-
-	FMaterialLayersFunctions LayersValue;
-	FGuid TempGuid(0, 0, 0, 0);
 
-	// TODO: Ben.Ingram / Cesar.Castro verify this.
-	bool bMaterialLayerReplaced = false;
-	if (MaterialInstance->GetMaterialLayers(LayersValue))
+	// Set Preview Mesh
+	if (DataAsset->GetStaticMesh())
 	{
-		for (int32 LayerIdx = 0; LayerIdx < LayersValue.Layers.Num(); ++LayerIdx)
-		{
-			TObjectPtr<UMaterialFunctionInterface>& MaterialFunction = LayersValue.Layers[LayerIdx];
-			if (MaterialFunction.Get() == OldMaterialFunction)
-			{
-				MaterialFunction = NewMaterialFunction;
-				LayersValue.EditorOnly.LayerLinkStates[LayerIdx] = EMaterialLayerLinkState::UnlinkedFromParent;
-				bMaterialLayerReplaced = true;
-				break;
-			}
-
-		}
+		MaterialInstance->PreviewMesh = DataAsset->GetStaticMesh();
 	}
 
-	if (bMaterialLayerReplaced)
-	{
-		// TODO: Ben.Ingram / Cesar.Castro verify this.
-		MaterialInstance->SetMaterialLayers(LayersValue);
-
-		// Update Material
-		UMaterialEditingLibrary::UpdateMaterialInstance(MaterialInstance);
-
-		// Rebuild Material
-		UMaterialEditingLibrary::RebuildMaterialInstanceEditors(MaterialInstance->GetMaterial());
-	}
+	MaterialInstance->MarkPackageDirty();
 }
 
 
@@ -811,11 +839,16 @@ bool UAnimToTextureBPLibrary::CreateUVChannel(
 	check(MeshDescription);
 
 	// Check if UVChannel is being used by the Lightmap UV
-	/*if (StaticMesh->GetLightMapCoordinateIndex() == UVChannelIndex)
+	if (StaticMesh->IsSourceModelValid(LODIndex))
 	{
-		UE_LOG(LogTemp, Warning, TEXT("Invalid UVChannel: %i. Already used by LightMap"), UVChannelIndex);
-		return false;
-	}*/
+		const FStaticMeshSourceModel& SourceModel = StaticMesh->GetSourceModel(LODIndex);
+		if (SourceModel.BuildSettings.bGenerateLightmapUVs &&
+			SourceModel.BuildSettings.DstLightmapIndex == UVChannelIndex)
+		{
+			UE_LOG(LogTemp, Warning, TEXT("Invalid UVChannel: %i. Already used by LightMap"), UVChannelIndex);
+			return false;
+		}
+	}
 
 	// Add New UVChannel.
 	// UE_LOG(LogTemp, Warning, TEXT("UVChannel: %i. Number of existing UVChannels: %i"), UVChannelIndex, StaticMesh->GetNumUVChannels(LODIndex));
@@ -849,7 +882,7 @@ bool UAnimToTextureBPLibrary::CreateUVChannel(
 	}
 
 	// Set Full Precision UVs
-	SetFullPrecisionUVs(StaticMesh, true);
+	SetFullPrecisionUVs(StaticMesh, LODIndex, true);
 
 	if (StaticMesh->SetUVChannel(LODIndex, UVChannelIndex, TexCoords))
 	{		
@@ -899,24 +932,55 @@ bool UAnimToTextureBPLibrary::FindBestResolution(
 	return bValidResolution;
 };
 
-void UAnimToTextureBPLibrary::SetFullPrecisionUVs(UStaticMesh* StaticMesh, bool bFullPrecision)
+void UAnimToTextureBPLibrary::SetFullPrecisionUVs(UStaticMesh* StaticMesh, int32 LODIndex, bool bFullPrecision)
 {
-	int32 NumSourceModels = StaticMesh->GetNumSourceModels();
-	for (int32 LodIndex = 0; LodIndex < NumSourceModels; LodIndex++)
+	if (StaticMesh->IsSourceModelValid(LODIndex))
 	{
-		FStaticMeshSourceModel& SourceModel = StaticMesh->GetSourceModel(LodIndex);
+		FStaticMeshSourceModel& SourceModel = StaticMesh->GetSourceModel(LODIndex);
 		SourceModel.BuildSettings.bUseFullPrecisionUVs = bFullPrecision;
 	}
 }
 
-void UAnimToTextureBPLibrary::SetStaticMeshBoundsExtensions(
-	UStaticMesh* StaticMesh, 
-	const FVector& PositiveBoundsExtension, 
-	const FVector& NegativeBoundsExtension)
+void UAnimToTextureBPLibrary::SetBoundsExtensions(UStaticMesh* StaticMesh, const FVector& MinBBox, const FVector& SizeBBox)
 {
+	// Calculate MaxBBox
+	const FVector MaxBBox = SizeBBox + MinBBox;
+
 	if (IsValid(StaticMesh))
 	{
-		StaticMesh->SetPositiveBoundsExtension(PositiveBoundsExtension);
-		StaticMesh->SetNegativeBoundsExtension(NegativeBoundsExtension);
+		// Reset current extension bounds
+		const FVector PositiveBoundsExtension = StaticMesh->GetPositiveBoundsExtension();
+		const FVector NegativeBoundsExtension = StaticMesh->GetNegativeBoundsExtension();
+		
+		// Get current BoundingBox including extensions
+		FBox BoundingBox = StaticMesh->GetBoundingBox();
+		
+		// Remove extensions from BoundingBox
+		BoundingBox.Max -= PositiveBoundsExtension;
+		BoundingBox.Min += NegativeBoundsExtension;
+		
+		// Calculate New BoundingBox
+		FVector NewMaxBBox(
+			FMath::Max(BoundingBox.Max.X, MaxBBox.X),
+			FMath::Max(BoundingBox.Max.Y, MaxBBox.Y),
+			FMath::Max(BoundingBox.Max.Z, MaxBBox.Z)
+		);
+		
+		FVector NewMinBBox(
+			FMath::Min(BoundingBox.Min.X, MinBBox.X),
+			FMath::Min(BoundingBox.Min.Y, MinBBox.Y),
+			FMath::Min(BoundingBox.Min.Z, MinBBox.Z)
+		);
+
+		// Calculate New Extensions
+		FVector NewPositiveBoundsExtension = NewMaxBBox - BoundingBox.Max;
+		FVector NewNegativeBoundsExtension = BoundingBox.Min - NewMinBBox;
+				
+		// Update StaticMesh
+		StaticMesh->SetPositiveBoundsExtension(NewPositiveBoundsExtension);
+		StaticMesh->SetNegativeBoundsExtension(NewNegativeBoundsExtension);
+		StaticMesh->CalculateExtendedBounds();
+		StaticMesh->PostEditChange();		
+		StaticMesh->MarkPackageDirty();
 	}
 }
diff --git a/Engine/Plugins/Animation/AnimToTexture/Source/AnimToTextureEditor/Private/AnimToTextureSkeletalMesh.cpp b/Engine/Plugins/Animation/AnimToTexture/Source/AnimToTextureEditor/Private/AnimToTextureSkeletalMesh.cpp
--- a/Engine/Plugins/Animation/AnimToTexture/Source/AnimToTextureEditor/Private/AnimToTextureSkeletalMesh.cpp
+++ b/Engine/Plugins/Animation/AnimToTexture/Source/AnimToTextureEditor/Private/AnimToTextureSkeletalMesh.cpp
@@ -261,58 +261,6 @@ void ReduceSkinWeightsData(const TArray<TVertexSkinWeight<MAX_TOTAL_INFLUENCES>>
 	}
 }
 
-void GetFollowerBoneMap(const USkeletalMesh* LeaderSkeletalMesh, const USkeletalMesh* FollowerSkeletalMesh, TArray<int32>& FollowerBoneMap)
-{
-	const TArray<FMeshBoneInfo>& LeaderBoneInfo = LeaderSkeletalMesh->GetRefSkeleton().GetRawRefBoneInfo();
-	const TArray<FMeshBoneInfo>& FollowerBoneInfo = FollowerSkeletalMesh->GetRefSkeleton().GetRawRefBoneInfo();
-
-	// Initialize Array
-	const int32 LeaderNumBones = LeaderBoneInfo.Num();
-	const int32 FollowerNumBones = FollowerBoneInfo.Num();
-	FollowerBoneMap.Init(INDEX_NONE, FollowerNumBones);
-
-	for (int32 FollowerBoneIndex = 0; FollowerBoneIndex < FollowerNumBones; ++FollowerBoneIndex)
-	{
-		// Search by name-matching
-		for (int32 LeaderBoneIndex = 0; LeaderBoneIndex < LeaderNumBones; ++LeaderBoneIndex)
-		{
-			// found by name-matching
-			if (FollowerBoneInfo[FollowerBoneIndex].Name == LeaderBoneInfo[LeaderBoneIndex].Name)
-			{
-				FollowerBoneMap[FollowerBoneIndex] = LeaderBoneIndex;
-				break;
-			}
-		}
-
-		// Search in upstream hierarchy
-		if (FollowerBoneMap[FollowerBoneIndex] == INDEX_NONE)
-		{
-			bool FoundLeaderBone = false;
-			int32 ParentIndex = FollowerBoneInfo[FollowerBoneIndex].ParentIndex;
-			while (!FoundLeaderBone && ParentIndex != INDEX_NONE)
-			{
-				for (int32 LeaderBoneIndex = 0; LeaderBoneIndex < LeaderNumBones; ++LeaderBoneIndex)
-				{
-					if (FollowerBoneInfo[ParentIndex].Name == LeaderBoneInfo[LeaderBoneIndex].Name)
-					{
-						FollowerBoneMap[FollowerBoneIndex] = LeaderBoneIndex;
-						FoundLeaderBone = true;
-						break;
-					}
-				}
-
-				ParentIndex = FollowerBoneInfo[ParentIndex].ParentIndex;
-			}
-		}	
-
-		// FollowerBone not found in Leader
-		if (FollowerBoneMap[FollowerBoneIndex] == INDEX_NONE)
-		{
-			UE_LOG(LogTemp, Warning, TEXT("Unable to find BoneMapping for: %s"), *FollowerBoneInfo[FollowerBoneIndex].Name.ToString())
-		}
-	}
-}
-
 void DecomposeTransformation(const FTransform& Transform, FVector3f& Translation, FVector4& Rotation)
 {
 	// Get Translation
@@ -417,29 +365,10 @@ void GetReducedSkinWeightsData(
 	OutSkinWeights.SetNumUninitialized(StaticToSkelMapping.Num());
 
 	// Loop thru SkeletalMesh Weights
-	//uint16 InvalidMeshBoneIndex;
-	//bool bIsValidMeshBoneIndex = true;
 	for (int32 VertexIndex = 0; VertexIndex < StaticToSkelMapping.Num(); ++VertexIndex)
 	{
 		const int32 SkelVertexIndex = StaticToSkelMapping[VertexIndex];
 		OutSkinWeights[VertexIndex] = ReducedSkinWeights[SkelVertexIndex];
-
-		// MeshBoneIndex can not be higher than 255
-		/*for (uint16 MeshBoneIndex: OutSkinWeights[VertexIndex].MeshBoneIndices)
-		{
-			if (MeshBoneIndex > 255)
-			{
-				InvalidMeshBoneIndex = MeshBoneIndex;
-				bIsValidMeshBoneIndex = false;
-				break;
-			}
-		}
-
-		if (!bIsValidMeshBoneIndex)
-		{
-			UE_LOG(LogTemp, Warning, TEXT("Invalid SkinWeights on Vertex: %i. MeshBoneIndex: %i > 255"), VertexIndex, InvalidMeshBoneIndex)
-			break;
-		}*/
 	}
 }
 
diff --git a/Engine/Plugins/Animation/AnimToTexture/Source/AnimToTextureEditor/Private/AnimToTextureUtils.cpp b/Engine/Plugins/Animation/AnimToTexture/Source/AnimToTextureEditor/Private/AnimToTextureUtils.cpp
--- a/Engine/Plugins/Animation/AnimToTexture/Source/AnimToTextureEditor/Private/AnimToTextureUtils.cpp
+++ b/Engine/Plugins/Animation/AnimToTexture/Source/AnimToTextureEditor/Private/AnimToTextureUtils.cpp
@@ -55,8 +55,8 @@ namespace AnimToTexture_Private
 		const int32 NumVertices = SkinWeights.Num();
 
 		// Allocate PixelData.
-		TArray<FColor> Pixels;
-		Pixels.Init(FColor::Black, Height * Width);
+		TArray<FLowPrecision::ColorType> Pixels;
+		Pixels.Init(FLowPrecision::ColorType::Black, Height * Width);
 
 		for (int32 VertexIndex = 0; VertexIndex < NumVertices; ++VertexIndex)
 		{
diff --git a/Engine/Plugins/Animation/AnimToTexture/Source/AnimToTextureEditor/Public/AnimToTextureBPLibrary.h b/Engine/Plugins/Animation/AnimToTexture/Source/AnimToTextureEditor/Public/AnimToTextureBPLibrary.h
--- a/Engine/Plugins/Animation/AnimToTexture/Source/AnimToTextureEditor/Public/AnimToTextureBPLibrary.h
+++ b/Engine/Plugins/Animation/AnimToTexture/Source/AnimToTextureEditor/Public/AnimToTextureBPLibrary.h
@@ -41,19 +41,8 @@ class UAnimToTextureBPLibrary : public UBlueprintFunctionLibrary
 	*/
 	UFUNCTION(BlueprintCallable, meta = (Category = "AnimToTexture"))
 	static void UpdateMaterialInstanceFromDataAsset(UAnimToTextureDataAsset* DataAsset, class UMaterialInstanceConstant* MaterialInstance, 
-		const FAnimToTextureMaterialParamNames& ParamNames, const EMaterialParameterAssociation MaterialParameterAssociation = EMaterialParameterAssociation::GlobalParameter);
-
-	/**
-	 * Replaces material layer function
-	 */
-	UFUNCTION(BlueprintCallable, meta = (Category = "AnimToTexture"))
-	static void UpdateMaterialLayerFunction(
-		class UMaterialInstanceConstant* MaterialInstance, 
-		class UMaterialFunctionInterface* OldMaterialFunction, 
-		class UMaterialFunctionInterface* NewMaterialFunction);
-
-	UFUNCTION(BlueprintCallable, meta = (Category = "AnimToTexture"))
-	static void SetStaticMeshBoundsExtensions(UStaticMesh* StaticMesh, const FVector& PositiveBoundsExtension, const FVector& NegativeBoundsExtension);
+		const bool bAnimate=false, const EAnimToTextureNumBoneInfluences NumBoneInfluences = EAnimToTextureNumBoneInfluences::One,
+		const EMaterialParameterAssociation MaterialParameterAssociation = EMaterialParameterAssociation::LayerParameter);
 
 #endif // WITH_EDITOR
 
@@ -77,7 +66,10 @@ private:
 								   const int32 MaxHeight = 4096, const int32 MaxWidth = 4096, bool bEnforcePowerOfTwo = false);
 
 	/* Sets Static Mesh FullPrecisionUVs Property*/
-	static void SetFullPrecisionUVs(UStaticMesh* StaticMesh, bool bFullPrecision);
+	static void SetFullPrecisionUVs(UStaticMesh* StaticMesh, const int32 LODIndex, bool bFullPrecision=true);
+
+	/* Sets Static Mesh Bound Extensions */
+	static void SetBoundsExtensions(UStaticMesh* StaticMesh, const FVector& MinBBox, const FVector& SizeBBox);
 
 	/* Creates UV Coord with vertices */
 	static bool CreateUVChannel(UStaticMesh* StaticMesh, const int32 LODIndex, const int32 UVChannelIndex,
diff --git a/Engine/Plugins/Animation/AnimToTexture/Source/AnimToTextureEditor/Public/AnimToTextureSkeletalMesh.h b/Engine/Plugins/Animation/AnimToTexture/Source/AnimToTextureEditor/Public/AnimToTextureSkeletalMesh.h
--- a/Engine/Plugins/Animation/AnimToTexture/Source/AnimToTextureEditor/Public/AnimToTextureSkeletalMesh.h
+++ b/Engine/Plugins/Animation/AnimToTexture/Source/AnimToTextureEditor/Public/AnimToTextureSkeletalMesh.h
@@ -69,86 +69,12 @@ bool HasBone(const USkeletalMesh* SkeletalMesh, const FName& Bone);
    Only the RawBones are returned (no virtual bones) */
 void GetBoneNames(const USkeletalMesh* SkeletalMesh, TArray<FName>& OutNames);
 
-/** Returns a Mapping between Follower and Leader. The hierarchy is traversed upstream until a matching Bone is found. */
-void GetFollowerBoneMap(const USkeletalMesh* LeaderSkeletalMesh, const USkeletalMesh* FollowerSkeletalMesh, TArray<int32>& FollowerBoneMap);
-
 /* Decomposes Transform in Translation and AxisAndAngle */
 void DecomposeTransformation(const FTransform& Transform, FVector3f& Translation, FVector4& Rotation);
 void DecomposeTransformations(const TArray<FTransform>& Transforms, TArray<FVector3f>& Translations, TArray<FVector4>& Rotations);
 
 // ----------------------------------------------------------------------------
 
-//template<>
-//FORCEINLINE void GetMaxInfluences<1>(
-//	const TVertexSkinWeight<MAX_TOTAL_INFLUENCES>& InSkinWeights, TVertexSkinWeight<4>& OutSkinWeights)
-//{
-//	// Reset Values
-//	for (int32 Index = 0; Index < 4; ++Index)
-//	{
-//		OutSkinWeights.BoneWeights[Index] = 0;
-//		OutSkinWeights.BoneIndices[Index] = 0;
-//		OutSkinWeights.MeshBoneIndices[Index] = 0;
-//	}
-//
-//	// Set first influence to 1
-//	OutSkinWeights.BoneWeights[0] = 255; // TNumericLimits<uint8>::Max();
-//
-//	// Find 1st Influence (Highest)
-//	uint8 MaxWeight = TNumericLimits<uint8>::Min();
-//	for (int32 Index = 0; Index < MAX_TOTAL_INFLUENCES; ++Index)
-//	{
-//		if (InSkinWeights.BoneWeights[Index] > MaxWeight)
-//		{
-//			MaxWeight = InSkinWeights.BoneWeights[Index];
-//			OutSkinWeights.BoneIndices[0] = InSkinWeights.BoneIndices[Index];
-//			OutSkinWeights.MeshBoneIndices[0] = InSkinWeights.MeshBoneIndices[Index];
-//		}
-//	}
-//}
-//
-//template<>
-//FORCEINLINE void GetMaxInfluences<2>(
-//	const TVertexSkinWeight<MAX_TOTAL_INFLUENCES>& InSkinWeights, TVertexSkinWeight<4>& OutSkinWeights)
-//{
-//	// Reset Values
-//	for (int32 Index = 0; Index < 4; ++Index)
-//	{
-//		OutSkinWeights.BoneIndices[Index] = 0;
-//		OutSkinWeights.BoneWeights[Index] = 0; 
-//		OutSkinWeights.MeshBoneIndices[Index] = 0;
-//	}
-//	
-//	// Find 1st Influence (Highest)
-//	for (int32 Index = 0; Index < MAX_TOTAL_INFLUENCES; ++Index)
-//	{
-//		if (InSkinWeights.BoneWeights[Index] > OutSkinWeights.BoneWeights[0])
-//		{
-//			OutSkinWeights.BoneWeights[0] = InSkinWeights.BoneWeights[Index];
-//			OutSkinWeights.BoneIndices[0] = InSkinWeights.BoneIndices[Index];
-//			OutSkinWeights.MeshBoneIndices[0] = InSkinWeights.MeshBoneIndices[Index];
-//		}
-//	}
-//
-//	// Find 2nd Influence
-//	for (int32 Index = 0; Index < MAX_TOTAL_INFLUENCES; ++Index)
-//	{
-//		if (InSkinWeights.BoneIndices[Index] != OutSkinWeights.BoneIndices[0])
-//		{
-//			if (InSkinWeights.BoneWeights[Index] > OutSkinWeights.BoneWeights[1])
-//			{
-//				OutSkinWeights.BoneWeights[1] = InSkinWeights.BoneWeights[Index];
-//				OutSkinWeights.BoneIndices[1] = InSkinWeights.BoneIndices[Index];
-//				OutSkinWeights.MeshBoneIndices[1] = InSkinWeights.MeshBoneIndices[Index];
-//			}
-//		}
-//	}
-//
-//	// Normalize Weighs
-//	const float NormFactor = 255.0f / (OutSkinWeights.BoneWeights[0] + OutSkinWeights.BoneWeights[1]);
-//	OutSkinWeights.BoneWeights[0] = FMath::RoundToInt(OutSkinWeights.BoneWeights[0] * NormFactor);
-//	OutSkinWeights.BoneWeights[1] = FMath::RoundToInt(OutSkinWeights.BoneWeights[1] * NormFactor);
-//}
-
 FORCEINLINE void GetMaxFourInfluences(
 	const TVertexSkinWeight<MAX_TOTAL_INFLUENCES>& InSkinWeights, TVertexSkinWeight<4>& OutSkinWeights)
 {
diff --git a/Engine/Plugins/Animation/AnimToTexture/Source/AnimToTextureEditor/Public/AnimToTextureUtils.h b/Engine/Plugins/Animation/AnimToTexture/Source/AnimToTextureEditor/Public/AnimToTextureUtils.h
--- a/Engine/Plugins/Animation/AnimToTexture/Source/AnimToTextureEditor/Public/AnimToTextureUtils.h
+++ b/Engine/Plugins/Animation/AnimToTexture/Source/AnimToTextureEditor/Public/AnimToTextureUtils.h
@@ -26,7 +26,7 @@ namespace AnimToTexture_Private
 		static constexpr EPixelFormat PixelFormat = EPixelFormat::PF_B8G8R8A8;
 		static constexpr ETextureSourceFormat TextureSourceFormat = ETextureSourceFormat::TSF_BGRA8;
 		static constexpr TextureCompressionSettings CompressionSettings = TextureCompressionSettings::TC_VectorDisplacementmap;
-		static constexpr ColorType DefaultColor = FColor(0, 0, 0);
+		static constexpr ColorType DefaultColor = FColor(0, 0, 0);		
 	};
 
 	struct FHighPrecision
@@ -185,7 +185,7 @@ FORCEINLINE bool AnimToTexture_Private::WriteVectorsToTexture(const TArray<V>& V
 	for (int32 Frame = 0; Frame < NumFrames; ++Frame)
 	{
 		const int32 BlockStart = RowsPerFrame * Width * Frame;
-		
+
 		// Set Data.
 		for (int32 Index = 0; Index < NumElements; ++Index)
 		{
@@ -202,9 +202,9 @@ FORCEINLINE bool AnimToTexture_Private::WriteVectorsToTexture(const TArray<V>& V
 
 template<class TextureSettings>
 FORCEINLINE bool AnimToTexture_Private::WriteToTexture(
-	UTexture2D* Texture, 
-	const uint32 Height, 
-	const uint32 Width, 
+	UTexture2D* Texture,
+	const uint32 Height,
+	const uint32 Width,
 	const TArray<typename TextureSettings::ColorType>& Pixels)
 {
 	if (!Texture)
@@ -237,9 +237,11 @@ FORCEINLINE bool AnimToTexture_Private::WriteToTexture(
 	else
 	{
 		Mip = &PlatformData->Mips[0];
+
 	}
 	Mip->SizeX = Width;
 	Mip->SizeY = Height;
+	Mip->SizeZ = 1;
 	
 	// ------------------------------------------------------------------------
 	// Lock the Mipmap data so it can be modified
