AnimToTexture:
- Code cleanup
- Barycentric mapping
- Simplified bone weight reduction

diff --git a/Engine/Plugins/Animation/AnimToTexture/Source/AnimToTexture/Public/AnimToTextureDataAsset.h b/Engine/Plugins/Animation/AnimToTexture/Source/AnimToTexture/Public/AnimToTextureDataAsset.h
--- a/Engine/Plugins/Animation/AnimToTexture/Source/AnimToTexture/Public/AnimToTextureDataAsset.h
+++ b/Engine/Plugins/Animation/AnimToTexture/Source/AnimToTexture/Public/AnimToTextureDataAsset.h
@@ -4,9 +4,7 @@
 
 #include "Engine/DataAsset.h"
 #include "Animation/AnimInstance.h"
-#include "Engine/SkeletalMesh.h"
 #include "Engine/StaticMesh.h"
-#include "Engine/Texture2D.h"
 #include "AnimToTextureDataAsset.generated.h"
 
 class USkeletalMesh;
@@ -93,13 +91,13 @@ struct FAnimInfo
 {
 	GENERATED_BODY()
 
-	UPROPERTY(VisibleAnywhere, Category = Default)
+	UPROPERTY(VisibleAnywhere, Category = Default, BlueprintReadOnly)
 	int32 NumFrames = 0;
 
-	UPROPERTY(VisibleAnywhere, Category = Default)
+	UPROPERTY(VisibleAnywhere, Category = Default, BlueprintReadOnly)
 	int32 AnimStart = 0;
 
-	UPROPERTY(EditAnywhere, Category = Default)
+	UPROPERTY(EditAnywhere, Category = Default, BlueprintReadOnly)
 	bool bLooping = true;
 };
 
@@ -165,9 +163,8 @@ public:
 
 	/**
 	* Texture Precision
-	* 16 bits is only supported for power-of-2 resolutions.
 	*/
-	UPROPERTY(EditAnywhere, BlueprintReadWrite, Category = "Texture", meta = (EditCondition="bEnforcePowerOfTwo"))
+	UPROPERTY(EditAnywhere, BlueprintReadWrite, Category = "Texture")
 	EAnimToTexturePrecision Precision = EAnimToTexturePrecision::EightBits;
 
 	/**
diff --git a/Engine/Plugins/Animation/AnimToTexture/Source/AnimToTextureEditor/AnimToTextureEditor.Build.cs b/Engine/Plugins/Animation/AnimToTexture/Source/AnimToTextureEditor/AnimToTextureEditor.Build.cs
--- a/Engine/Plugins/Animation/AnimToTexture/Source/AnimToTextureEditor/AnimToTextureEditor.Build.cs
+++ b/Engine/Plugins/Animation/AnimToTexture/Source/AnimToTextureEditor/AnimToTextureEditor.Build.cs
@@ -42,6 +42,7 @@ public class AnimToTextureEditor : ModuleRules
 				"AnimToTexture",
 				"RawMesh",
 				"MeshDescription",
+				"StaticMeshDescription",
 				"UnrealEd",
 			}
 			);
diff --git a/Engine/Plugins/Animation/AnimToTexture/Source/AnimToTextureEditor/Private/AnimToTextureBPLibrary.cpp b/Engine/Plugins/Animation/AnimToTexture/Source/AnimToTextureEditor/Private/AnimToTextureBPLibrary.cpp
--- a/Engine/Plugins/Animation/AnimToTexture/Source/AnimToTextureEditor/Private/AnimToTextureBPLibrary.cpp
+++ b/Engine/Plugins/Animation/AnimToTexture/Source/AnimToTextureEditor/Private/AnimToTextureBPLibrary.cpp
@@ -1,31 +1,24 @@
 // Copyright Epic Games, Inc. All Rights Reserved.
 
 #include "AnimToTextureBPLibrary.h"
-#include "Animation/Skeleton.h"
 #include "AnimToTexture.h"
 #include "AnimToTextureUtils.h"
 #include "AnimToTextureSkeletalMesh.h"
-#include "Engine/SkeletalMesh.h"
 #include "EvaluateSequenceAnimInstance.h"
 
-#include "Editor.h"
 #include "LevelEditor.h"
 #include "RawMesh.h"
 #include "MeshUtilities.h"
-#include "AssetRegistry/AssetRegistryModule.h"
+#include "Engine/SkeletalMesh.h"
 #include "Components/SkeletalMeshComponent.h"
-#include "Components/SkinnedMeshComponent.h"
-#include "Rendering/SkeletalMeshRenderData.h"
-#include "Rendering/SkeletalMeshModel.h"
-#include "Animation/AnimBlueprint.h"
 #include "Animation/AnimSequence.h"
 #include "Math/Vector.h"
 #include "Math/NumericLimits.h"
-#include "Engine/SkeletalMeshSocket.h"
 #include "MeshDescription.h"
 #include "Materials/MaterialInstanceConstant.h"
 #include "MaterialEditingLibrary.h"
-#include "EditorSupportDelegates.h"
+
+using namespace AnimToTexture_Private;
 
 UAnimToTextureBPLibrary::UAnimToTextureBPLibrary(const FObjectInitializer& ObjectInitializer)
 : Super(ObjectInitializer)
@@ -33,8 +26,10 @@ UAnimToTextureBPLibrary::UAnimToTextureBPLibrary(const FObjectInitializer& Objec
 
 }
 
-void UAnimToTextureBPLibrary::AnimationToTexture(UAnimToTextureDataAsset* DataAsset, const FTransform RootTransform)
+void UAnimToTextureBPLibrary::AnimationToTexture(UAnimToTextureDataAsset* DataAsset, const FTransform RootTransform, bool& bSuccess)
 {
+	bSuccess = false; 
+
 	if (!DataAsset)
 	{
 		return;
@@ -43,17 +38,38 @@ void UAnimToTextureBPLibrary::AnimationToTexture(UAnimToTextureDataAsset* DataAs
 	// Reset DataAsset Info Values
 	DataAsset->ResetInfo();
 
-	if (!DataAsset->GetSkeletalMesh() || !DataAsset->GetStaticMesh())
+	// Check StaticMesh
+	if (!DataAsset->GetStaticMesh())
+	{
+		UE_LOG(LogTemp, Warning, TEXT("Invalid StaticMesh"));
+		return;
+	}
+
+	// Check SkeletalMesh
+	if (!DataAsset->GetSkeletalMesh())
+	{
+		UE_LOG(LogTemp, Warning, TEXT("Invalid SkeletalMesh"));
+		return;
+	}
+
+	// Check Skeleton
+	if (!DataAsset->GetSkeletalMesh()->GetSkeleton())
+	{
+		UE_LOG(LogTemp, Warning, TEXT("Invalid SkeletalMesh. No valid Skeleton found"));
+		return;
+	}
+
+	// Check StaticMesh LOD
+	if (!DataAsset->GetStaticMesh()->IsSourceModelValid(DataAsset->StaticLODIndex))
 	{
-		UE_LOG(LogTemp, Warning, TEXT("Invalid Meshes"));
+		UE_LOG(LogTemp, Warning, TEXT("Invalid StaticMesh LOD Index"));
 		return;
 	}
 
-	// Check LODs
-	if (!DataAsset->GetSkeletalMesh()->IsValidLODIndex(DataAsset->SkeletalLODIndex) || 
-		DataAsset->StaticLODIndex >= DataAsset->GetStaticMesh()->GetNumLODs())
+	// Check SkeletalMesh LOD
+	if (!DataAsset->GetSkeletalMesh()->IsValidLODIndex(DataAsset->SkeletalLODIndex))
 	{
-		UE_LOG(LogTemp, Warning, TEXT("Invalid LOD Indices"));
+		UE_LOG(LogTemp, Warning, TEXT("Invalid SkeletalMesh LOD Index"));
 		return;
 	}
 
@@ -61,7 +77,7 @@ void UAnimToTextureBPLibrary::AnimationToTexture(UAnimToTextureDataAsset* DataAs
 	bool bValidSocket = false;
 	if (DataAsset->AttachToSocket.IsValid() && !DataAsset->AttachToSocket.IsNone())
 	{
-		if (AnimToTexture_Private::HasBone(DataAsset->GetSkeletalMesh(), DataAsset->AttachToSocket))
+		if (HasBone(DataAsset->GetSkeletalMesh(), DataAsset->AttachToSocket))
 		{
 			bValidSocket = true;
 		}
@@ -71,53 +87,52 @@ void UAnimToTextureBPLibrary::AnimationToTexture(UAnimToTextureDataAsset* DataAs
 			return;
 		}
 	}
-	
 	if (bValidSocket && DataAsset->Mode == EAnimToTextureMode::Vertex)
 	{
 		UE_LOG(LogTemp, Warning, TEXT("Unable to use Socket in Vertex Mode. Use Bone Mode instead."));
 		return;
 	}
 
+	// Check if UVChannel is being used by the Lightmap UV
+	const FStaticMeshSourceModel& SourceModel = DataAsset->GetStaticMesh()->GetSourceModel(DataAsset->StaticLODIndex);
+	if (SourceModel.BuildSettings.bGenerateLightmapUVs &&
+		SourceModel.BuildSettings.DstLightmapIndex == DataAsset->UVChannel)
+	{
+		UE_LOG(LogTemp, Warning, TEXT("Invalid UVChannel: %i. Already used by LightMap"), DataAsset->UVChannel);
+		return;
+	}
+
 	// ---------------------------------------------------------------------------		
 	// Get Meshes Vertices and Mapping.
-	// NOTE: We need to create a Mapping between the StaticMesh|MeshDescription and the SkeletalMesh
+	// NOTE: We need to create a Mapping between the StaticMesh and the SkeletalMesh
 	//       Since they dont have same number of points.
-	//	
+	//
+
+	// Get Static Vertices 
+	// NOTE: these are the RenderData Vertices.
 	TArray<FVector3f> Vertices;
-	TArray<FVector3f> SkelVertices;
-	TArray<int32> StaticToSkelMapping;
+	TArray<FVector3f> Normals;
+	const int32 NumVertices = GetVertices(DataAsset->GetStaticMesh(), DataAsset->StaticLODIndex, Vertices, Normals);
 
-	AnimToTexture_Private::GetStaticToSkeletalMapping(
-		*DataAsset->GetStaticMesh(), DataAsset->StaticLODIndex,
-		*DataAsset->GetSkeletalMesh(), DataAsset->SkeletalLODIndex,
-		Vertices, SkelVertices,
-		StaticToSkelMapping);
+	// Get SkeletalMesh Vertices 
+	TArray<FVector3f> SkelVertices;
+	GetVertices(DataAsset->GetSkeletalMesh(), DataAsset->SkeletalLODIndex, SkelVertices);
 	
-	// Set Static Number of Vertices 
-	// NOTE: these are the MeshDescription Vertices.
-	const int32 NumVertices = Vertices.Num();
+	// Get Mapping
+	TArray<FVertexToMeshMapping> Mapping;
+	FVertexToMeshMapping::Create(
+		DataAsset->GetStaticMesh(), DataAsset->StaticLODIndex,
+		DataAsset->GetSkeletalMesh(), DataAsset->SkeletalLODIndex, Mapping);
 
 	// --------------------------------------------------------------------------
-
-	if (GEditor == nullptr)
-	{
-		return;
-	}
-
+		
 	// Create Temp Actor
+	check(GEditor);
 	UWorld* World = GEditor->GetEditorWorldContext().World();
-
-	if (World == nullptr)
-	{
-		return;
-	}
+	check(World);
 
 	AActor* Actor = World->SpawnActor<AActor>();
-
-	if (Actor == nullptr)
-	{
-		return;
-	}
+	check(Actor);
 
 	// Create Temp SkeletalMesh Component
 	USkeletalMeshComponent* MeshComponent = NewObject<USkeletalMeshComponent>(Actor);
@@ -142,19 +157,22 @@ void UAnimToTextureBPLibrary::AnimationToTexture(UAnimToTextureDataAsset* DataAs
 	
 	if (DataAsset->Mode == EAnimToTextureMode::Bone)
 	{
-		// Get Number of RawBones (no virtual)
-		NumBones = AnimToTexture_Private::GetNumBones(MeshComponent->GetSkeletalMeshAsset());
+		// Gets Ref Bone Position and Rotations.
+		NumBones = GetRefBonePositionsAndRotations(MeshComponent,
+			BoneRefPositions, BoneRefRotations);
 
-		// Get Raw Ref Bone (no virtual)
-		TArray<FTransform> RefBoneTransforms;
-		AnimToTexture_Private::GetRefBoneTransforms(MeshComponent->GetSkeletalMeshAsset(), RefBoneTransforms);
-		AnimToTexture_Private::DecomposeTransformations(RefBoneTransforms, BoneRefPositions, BoneRefRotations);
+		// NOTE: there is a limitation with the number of bones atm.
+		if (NumBones > 256)
+		{
+			UE_LOG(LogTemp, Warning, TEXT("Invalid Number of Bones. There is a maximum of 256 bones"))
+			return;
+		}
 
 		// Get Bone Names (no virtual)
-		AnimToTexture_Private::GetBoneNames(MeshComponent->GetSkeletalMeshAsset(), BoneNames);
+		GetBoneNames(MeshComponent->GetSkeletalMeshAsset(), BoneNames);
 
 		// Make sure array sizes are correct.
-		check(RefBoneTransforms.Num() == NumBones && BoneNames.Num() == NumBones);
+		check(BoneNames.Num() == NumBones);
 
 		// Check if Socket is in BoneNames
 		if (bValidSocket && !BoneNames.Find(DataAsset->AttachToSocket, SocketIndex))
@@ -191,7 +209,7 @@ void UAnimToTextureBPLibrary::AnimationToTexture(UAnimToTextureDataAsset* DataAs
 		// Make sure SkeletalMesh is compatible with AnimSequence
 		if (!MeshComponent->GetSkeletalMeshAsset()->GetSkeleton()->IsCompatibleForEditor(AnimSequence->GetSkeleton()))
 		{
-			UE_LOG(LogTemp, Warning, TEXT("Incompatible AnimSequence: %s for given SkeletalMesh: %s"), *AnimSequence->GetFName().ToString(), *MeshComponent->GetSkeletalMeshAsset()->GetFName().ToString());
+			UE_LOG(LogTemp, Warning, TEXT("Invalid AnimSequence: %s for given SkeletalMesh: %s"), *AnimSequence->GetFName().ToString(), *MeshComponent->GetSkeletalMeshAsset()->GetFName().ToString());
 			continue;
 		}
 		// Set AnimSequence
@@ -226,7 +244,7 @@ void UAnimToTextureBPLibrary::AnimationToTexture(UAnimToTextureDataAsset* DataAs
 		// ---------------------------------------------------------------------------
 		// 
 		float Time = 0.f;
-		float EndTime = AnimSequence->GetTimeAtFrame(EndFrame);
+		const float EndTime = AnimSequence->GetTimeAtFrame(EndFrame);
 
 		int32 SampleIndex = 0;
 
@@ -239,8 +257,8 @@ void UAnimToTextureBPLibrary::AnimationToTexture(UAnimToTextureDataAsset* DataAs
 			MeshComponent->SetPosition(Time);
 
 			// Update SkelMesh Animation.
-			MeshComponent->TickAnimation(0.0f, false /*bNeedsValidRootMotion*/);
-			// MeshComponent->TickComponent(0.0f, ELevelTick::LEVELTICK_All, nullptr);
+			MeshComponent->TickAnimation(0.f, false /*bNeedsValidRootMotion*/);
+			// MeshComponent->TickComponent(0.f, ELevelTick::LEVELTICK_All, nullptr);
 			MeshComponent->RefreshBoneTransforms(nullptr /*TickFunction*/);
 			
 			// ---------------------------------------------------------------------------
@@ -248,23 +266,17 @@ void UAnimToTextureBPLibrary::AnimationToTexture(UAnimToTextureDataAsset* DataAs
 			//
 			if (DataAsset->Mode == EAnimToTextureMode::Vertex)
 			{
-				TArray<FVector3f> Positions;
-				TArray<FVector3f> Normals;
-				AnimToTexture_Private::GetSkinnedVertices(MeshComponent, DataAsset->SkeletalLODIndex,
-					Positions, Normals);
-
-				// Loop thru static vertices and find the mapped SkeletalMesh Vertex
-				for (int32 VertexIndex = 0; VertexIndex < NumVertices; ++VertexIndex)
-				{
-					const int32 SkelVertexIndex = StaticToSkelMapping[VertexIndex];
-
-					// Delta & Normal
-					const FVector3f VertexDelta = (FVector3f)RootTransform.TransformPosition((FVector)Positions[SkelVertexIndex]) - SkelVertices[SkelVertexIndex];
-					const FVector3f VertexNormal = (FVector3f)RootTransform.TransformVector((FVector)Normals[SkelVertexIndex]);
+				TArray<FVector3f> VertexFrameDeltas;
+				TArray<FVector3f> VertexFrameNormals;
+				
+				GetVertexDeltasAndNormals(MeshComponent, DataAsset->SkeletalLODIndex,
+					Vertices, Normals, Mapping,
+					RootTransform,
+					VertexFrameDeltas, VertexFrameNormals);
+					
+				VertexDeltas.Append(VertexFrameDeltas);
+				VertexNormals.Append(VertexFrameNormals);
 
-					VertexDeltas.Add(VertexDelta);
-					VertexNormals.Add(VertexNormal);
-				}
 			} // End Vertex Mode
 
 			// ---------------------------------------------------------------------------
@@ -272,41 +284,17 @@ void UAnimToTextureBPLibrary::AnimationToTexture(UAnimToTextureDataAsset* DataAs
 			//
 			else if (DataAsset->Mode == EAnimToTextureMode::Bone)
 			{
-				// Get Relative Transforms
-				// Note: Size is of Raw bones in SkeletalMesh. These are the original/raw bones of the asset, without Virtual Bones.
-				TArray<FMatrix44f> RefToLocals;
-				MeshComponent->CacheRefToLocalMatrices(RefToLocals);
-				
-				// check size
-				check(RefToLocals.Num() == NumBones);
+				TArray<FVector3f> BoneFramePositions;
+				TArray<FVector4> BoneFrameRotations;
 
-				// Get Component Space Transforms
-				// Note returns all transforms, including VirtualBones
-				const TArray<FTransform>& CompSpaceTransforms = MeshComponent->GetComponentSpaceTransforms();
-				check(CompSpaceTransforms.Num() >= RefToLocals.Num());
+				GetBonePositionsAndRotations(MeshComponent, BoneRefPositions,
+					BoneFramePositions, BoneFrameRotations);
+
+				BonePositions.Append(BoneFramePositions);
+				BoneRotations.Append(BoneFrameRotations);
 
-				for (int32 BoneIndex = 0; BoneIndex < NumBones; ++BoneIndex)
-				{
-					// Decompose Transformation (ComponentSpace)
-					const FTransform& CompSpaceTransform = CompSpaceTransforms[BoneIndex];
-					FVector3f BonePosition;
-					FVector4 BoneRotation;
-					AnimToTexture_Private::DecomposeTransformation(CompSpaceTransform, BonePosition, BoneRotation);
-
-					// Position Delta (from RefPose)
-					const FVector3f Delta = BonePosition - BoneRefPositions[BoneIndex];
-
-					// Decompose Transformation (Relative to RefPose)
-					FVector3f BoneRelativePosition;
-					FVector4 BoneRelativeRotation;
-					FMatrix RefToLocalMatrix = FMatrix(RefToLocals[BoneIndex]);
-					const FTransform RelativeTransform(RefToLocalMatrix);
-					AnimToTexture_Private::DecomposeTransformation(RelativeTransform, BoneRelativePosition, BoneRelativeRotation);
-
-					BonePositions.Add(Delta);
-					BoneRotations.Add(BoneRelativeRotation);
-				}
 			} // End Bone Mode
+
 		} // End Frame
 
 		// Store Anim Info Data
@@ -356,22 +344,25 @@ void UAnimToTextureBPLibrary::AnimationToTexture(UAnimToTextureDataAsset* DataAs
 			NormalizedVertexDeltas, NormalizedVertexNormals);
 
 		// Write Textures
-		if (DataAsset->Precision == EAnimToTexturePrecision::SixteenBits && DataAsset->bEnforcePowerOfTwo)
+		if (DataAsset->Precision == EAnimToTexturePrecision::SixteenBits)
 		{
-			AnimToTexture_Private::WriteVectorsToTexture<FVector3f, AnimToTexture_Private::FHighPrecision>(NormalizedVertexDeltas, DataAsset->NumFrames, DataAsset->VertexRowsPerFrame, Height, Width, DataAsset->GetVertexPositionTexture());
-			AnimToTexture_Private::WriteVectorsToTexture<FVector3f, AnimToTexture_Private::FHighPrecision>(NormalizedVertexNormals, DataAsset->NumFrames, DataAsset->VertexRowsPerFrame, Height, Width, DataAsset->GetVertexNormalTexture());
+			WriteVectorsToTexture<FVector3f, FHighPrecision>(NormalizedVertexDeltas, DataAsset->NumFrames, DataAsset->VertexRowsPerFrame, Height, Width, DataAsset->GetVertexPositionTexture());
+			WriteVectorsToTexture<FVector3f, FHighPrecision>(NormalizedVertexNormals, DataAsset->NumFrames, DataAsset->VertexRowsPerFrame, Height, Width, DataAsset->GetVertexNormalTexture());
 		}
 		else
 		{
-			AnimToTexture_Private::WriteVectorsToTexture<FVector3f, AnimToTexture_Private::FLowPrecision>(NormalizedVertexDeltas, DataAsset->NumFrames, DataAsset->VertexRowsPerFrame, Height, Width, DataAsset->GetVertexPositionTexture());
-			AnimToTexture_Private::WriteVectorsToTexture<FVector3f, AnimToTexture_Private::FLowPrecision>(NormalizedVertexNormals, DataAsset->NumFrames, DataAsset->VertexRowsPerFrame, Height, Width, DataAsset->GetVertexNormalTexture());
-		}
-			
+			WriteVectorsToTexture<FVector3f, FLowPrecision>(NormalizedVertexDeltas, DataAsset->NumFrames, DataAsset->VertexRowsPerFrame, Height, Width, DataAsset->GetVertexPositionTexture());
+			WriteVectorsToTexture<FVector3f, FLowPrecision>(NormalizedVertexNormals, DataAsset->NumFrames, DataAsset->VertexRowsPerFrame, Height, Width, DataAsset->GetVertexNormalTexture());
+		}		
+
 		// Add Vertex UVChannel
 		CreateUVChannel(DataAsset->GetStaticMesh(), DataAsset->StaticLODIndex, DataAsset->UVChannel, Height, Width);
-		
+
 		// Update Bounds
 		SetBoundsExtensions(DataAsset->GetStaticMesh(), DataAsset->VertexMinBBox, DataAsset->VertexSizeBBox);
+
+		// Done with StaticMesh
+		DataAsset->GetStaticMesh()->PostEditChange();
 	}
 
 	// ---------------------------------------------------------------------------
@@ -381,8 +372,6 @@ void UAnimToTextureBPLibrary::AnimationToTexture(UAnimToTextureDataAsset* DataAs
 		// Find Best Resolution for Bone Data
 		int32 Height, Width;
 
-		// NOTE: If NumBones are > 256, you might need to use a simpler skeleton
-
 		// Note we are adding +1 frame for the ref pose
 		if (!FindBestResolution(DataAsset->NumFrames + 1, NumBones, 
 								Height, Width, DataAsset->BoneRowsPerFrame, 
@@ -403,15 +392,15 @@ void UAnimToTextureBPLibrary::AnimationToTexture(UAnimToTextureDataAsset* DataAs
 				NormalizedBonePositions, NormalizedBoneRotations);
 
 			// Write Textures
-			if (DataAsset->Precision == EAnimToTexturePrecision::SixteenBits && DataAsset->bEnforcePowerOfTwo)
+			if (DataAsset->Precision == EAnimToTexturePrecision::SixteenBits)
 			{
-				AnimToTexture_Private::WriteVectorsToTexture<FVector3f, AnimToTexture_Private::FHighPrecision>(NormalizedBonePositions, DataAsset->NumFrames + 1, DataAsset->BoneRowsPerFrame, Height, Width, DataAsset->GetBonePositionTexture());
-				AnimToTexture_Private::WriteVectorsToTexture<FVector4, AnimToTexture_Private::FHighPrecision>(NormalizedBoneRotations, DataAsset->NumFrames + 1, DataAsset->BoneRowsPerFrame, Height, Width, DataAsset->GetBoneRotationTexture());
+				WriteVectorsToTexture<FVector3f, FHighPrecision>(NormalizedBonePositions, DataAsset->NumFrames + 1, DataAsset->BoneRowsPerFrame, Height, Width, DataAsset->GetBonePositionTexture());
+				WriteVectorsToTexture<FVector4, FHighPrecision>(NormalizedBoneRotations, DataAsset->NumFrames + 1, DataAsset->BoneRowsPerFrame, Height, Width, DataAsset->GetBoneRotationTexture());
 			}
 			else
 			{
-				AnimToTexture_Private::WriteVectorsToTexture<FVector3f, AnimToTexture_Private::FLowPrecision>(NormalizedBonePositions, DataAsset->NumFrames + 1, DataAsset->BoneRowsPerFrame, Height, Width, DataAsset->GetBonePositionTexture());
-				AnimToTexture_Private::WriteVectorsToTexture<FVector4, AnimToTexture_Private::FLowPrecision>(NormalizedBoneRotations, DataAsset->NumFrames + 1, DataAsset->BoneRowsPerFrame, Height, Width, DataAsset->GetBoneRotationTexture());
+				WriteVectorsToTexture<FVector3f, FLowPrecision>(NormalizedBonePositions, DataAsset->NumFrames + 1, DataAsset->BoneRowsPerFrame, Height, Width, DataAsset->GetBonePositionTexture());
+				WriteVectorsToTexture<FVector4, FLowPrecision>(NormalizedBoneRotations, DataAsset->NumFrames + 1, DataAsset->BoneRowsPerFrame, Height, Width, DataAsset->GetBoneRotationTexture());
 			}
 		}
 
@@ -428,23 +417,29 @@ void UAnimToTextureBPLibrary::AnimationToTexture(UAnimToTextureDataAsset* DataAs
 
 		// Write Weights Texture
 		{
-			TArray<AnimToTexture_Private::TVertexSkinWeight<4>> SkinWeights;
+			TArray<TVertexSkinWeight<4>> SkinWeights;
 
 			// Store Influence Weights
 			if (!bValidSocket)
 			{
-				// Get SkinWeights mapped to StaticMesh
-				AnimToTexture_Private::GetReducedSkinWeightsData(
-					*DataAsset->GetSkeletalMesh(), DataAsset->SkeletalLODIndex,
-					StaticToSkelMapping, 
-					SkinWeights);
+				// Get SkeletalMesh Weights
+				TArray<VertexSkinWeightMax> SkeletalMeshSkinWeights;
+				GetSkinWeights(DataAsset->GetSkeletalMesh(), DataAsset->SkeletalLODIndex, SkeletalMeshSkinWeights);
+
+				// Interpolate Weights with Mapping
+				TArray<VertexSkinWeightMax> StaticMeshSkinWeights;
+				FVertexToMeshMapping::InterpolateSkinWeights(Mapping, SkeletalMeshSkinWeights, StaticMeshSkinWeights);
+				
+				// Reduce SkeletalMesh Weights to 4 highest influences.
+				ReduceSkinWeights(StaticMeshSkinWeights, SkinWeights);
 			}
+
 			// If Valid Socket, set all influences to same index.
 			else
 			{
 				// Set all indices and weights to same SocketIndex
 				SkinWeights.SetNumUninitialized(NumVertices);
-				for (AnimToTexture_Private::TVertexSkinWeight<4>& SkinWeight : SkinWeights)
+				for (TVertexSkinWeight<4>& SkinWeight : SkinWeights)
 				{
 					SkinWeight.BoneWeights = TStaticArray<uint8, 4>(InPlace, 255);
 					SkinWeight.MeshBoneIndices = TStaticArray<uint16, 4>(InPlace, SocketIndex);
@@ -452,7 +447,7 @@ void UAnimToTextureBPLibrary::AnimationToTexture(UAnimToTextureDataAsset* DataAs
 			}
 
 			// Write Bone Weights Texture
-			AnimToTexture_Private::WriteSkinWeightsToTexture(SkinWeights,
+			WriteSkinWeightsToTexture(SkinWeights,
 				DataAsset->BoneWeightRowsPerFrame, Height, Width, DataAsset->GetBoneWeightTexture());
 		}
 
@@ -461,12 +456,122 @@ void UAnimToTextureBPLibrary::AnimationToTexture(UAnimToTextureDataAsset* DataAs
 
 		// Update Bounds
 		SetBoundsExtensions(DataAsset->GetStaticMesh(), DataAsset->BoneMinBBox, DataAsset->BoneSizeBBox);
+
+		// Done with StaticMesh
+		DataAsset->GetStaticMesh()->PostEditChange();
 	}
 
 	// ---------------------------------------------------------------------------
-	// Mark DataAsset as Dirty.
+	// Mark Packages dirty
 	//
 	DataAsset->MarkPackageDirty();
+	
+	// All good here !
+	bSuccess = true;
+}
+
+
+// 
+void UAnimToTextureBPLibrary::GetVertexDeltasAndNormals(const USkeletalMeshComponent* SkeletalMeshComponent, const int32 LODIndex, 
+	const TArray<FVector3f>& Vertices, const TArray<FVector3f>& Normals,  const TArray<FVertexToMeshMapping>& Mapping,
+	const FTransform RootTransform,
+	TArray<FVector3f>& OutVertexDeltas, TArray<FVector3f>& OutVertexNormals)
+{
+	OutVertexDeltas.Reset();
+	OutVertexNormals.Reset();
+
+	check(Vertices.Num() == Mapping.Num());
+
+	TArray<FVector3f> SkinnedVertices;
+	GetSkinnedVertices(SkeletalMeshComponent, LODIndex, SkinnedVertices);
+
+	// Loop thru static vertices and find the mapped SkeletalMesh Vertex
+	for (int32 VertexIndex = 0; VertexIndex < Vertices.Num(); VertexIndex++)
+	{
+		// Get Static/SkeletalMesh Mapping
+		const FVertexToMeshMapping& VertexMapping = Mapping[VertexIndex];
+
+		// Transform Vertex
+		// (assuming indices are valid)
+		const FIntVector3& Triangle = VertexMapping.Triangle;
+		const FVector3f& A = SkinnedVertices[Triangle.X];
+		const FVector3f& B = SkinnedVertices[Triangle.Y];
+		const FVector3f& C = SkinnedVertices[Triangle.Z];
+
+		const FVector3f VertexPosition = VertexMapping.TransformPosition(Vertices[VertexIndex], A, B, C);
+		const FVector3f VertexNormal = VertexMapping.TransformVector(Normals[VertexIndex], A, B, C);
+
+		// Transform Position and Delta with RootTransform
+		const FVector3f TransformedVertexDelta = ((FVector3f)RootTransform.TransformPosition((FVector)VertexPosition)) - Vertices[VertexIndex];
+		const FVector3f TransformedVertexNormal = (FVector3f)RootTransform.TransformVector((FVector)VertexNormal);
+		
+		OutVertexDeltas.Add(TransformedVertexDelta);
+		OutVertexNormals.Add(TransformedVertexNormal);
+	}
+}
+
+
+int32 UAnimToTextureBPLibrary::GetRefBonePositionsAndRotations(const USkeletalMeshComponent* SkeletalMeshComponent,
+	TArray<FVector3f>& OutBoneRefPositions, TArray<FVector4>& OutBoneRefRotations)
+{
+	OutBoneRefPositions.Reset();
+	OutBoneRefRotations.Reset();
+
+	// Get Number of RawBones (no virtual)
+	const int32 NumBones = GetNumBones(SkeletalMeshComponent->GetSkeletalMeshAsset());
+	
+	// Get Raw Ref Bone (no virtual)
+	TArray<FTransform> RefBoneTransforms;
+	GetRefBoneTransforms(SkeletalMeshComponent->GetSkeletalMeshAsset(), RefBoneTransforms);
+	DecomposeTransformations(RefBoneTransforms, OutBoneRefPositions, OutBoneRefRotations);
+
+	return NumBones;
+}
+
+
+int32 UAnimToTextureBPLibrary::GetBonePositionsAndRotations(const USkeletalMeshComponent* SkeletalMeshComponent, const TArray<FVector3f>& BoneRefPositions,
+	TArray<FVector3f>& BonePositions, TArray<FVector4>& BoneRotations)
+{
+	BonePositions.Reset();
+	BoneRotations.Reset();
+
+	// Get Relative Transforms
+	// Note: Size is of Raw bones in SkeletalMesh. These are the original/raw bones of the asset, without Virtual Bones.
+	TArray<FMatrix44f> RefToLocals;
+	SkeletalMeshComponent->CacheRefToLocalMatrices(RefToLocals);
+	const int32 NumBones = RefToLocals.Num();
+
+	// check size
+	check(NumBones == BoneRefPositions.Num());
+
+	// Get Component Space Transforms
+	// Note returns all transforms, including VirtualBones
+	const TArray<FTransform>& CompSpaceTransforms = SkeletalMeshComponent->GetComponentSpaceTransforms();
+	check(CompSpaceTransforms.Num() >= RefToLocals.Num());
+
+	for (int32 BoneIndex = 0; BoneIndex < NumBones; BoneIndex++)
+	{
+		// Decompose Transformation (ComponentSpace)
+		const FTransform& CompSpaceTransform = CompSpaceTransforms[BoneIndex];
+		FVector3f BonePosition;
+		FVector4 BoneRotation;
+		DecomposeTransformation(CompSpaceTransform, BonePosition, BoneRotation);
+
+		// Position Delta (from RefPose)
+		const FVector3f Delta = BonePosition - BoneRefPositions[BoneIndex];
+
+		// Decompose Transformation (Relative to RefPose)
+		FVector3f BoneRelativePosition;
+		FVector4 BoneRelativeRotation;
+		const FMatrix RefToLocalMatrix(RefToLocals[BoneIndex]);
+		const FTransform RelativeTransform(RefToLocalMatrix);
+		DecomposeTransformation(RelativeTransform, BoneRelativePosition, BoneRelativeRotation);
+
+		BonePositions.Add(Delta);
+		BoneRotations.Add(BoneRelativeRotation);
+	}
+
+	return NumBones;
 }
 
 
@@ -585,6 +690,33 @@ void UAnimToTextureBPLibrary::UpdateMaterialInstanceFromDataAsset(UAnimToTexture
 }
 
 
+bool UAnimToTextureBPLibrary::SetLightMapIndex(UStaticMesh* StaticMesh, const int32 LODIndex, const int32 LightmapIndex, bool bGenerateLightmapUVs)
+{
+	if (!StaticMesh)
+	{
+		return false;
+	}
+
+	if (LODIndex >= 0 && !StaticMesh->IsSourceModelValid(LODIndex))
+	{
+		return false;
+	}
+
+	// 
+	FStaticMeshSourceModel& SourceModel = StaticMesh->GetSourceModel(LODIndex);
+	SourceModel.BuildSettings.bGenerateLightmapUVs = bGenerateLightmapUVs;
+	SourceModel.BuildSettings.DstLightmapIndex = LightmapIndex;
+	StaticMesh->SetLightMapCoordinateIndex(LightmapIndex);
+
+	// Build mesh from source
+	StaticMesh->Build(false);
+	StaticMesh->PostEditChange();
+	StaticMesh->MarkPackageDirty();
+
+	return true;
+}
+
+
 UStaticMesh* UAnimToTextureBPLibrary::ConvertSkeletalMeshToStaticMesh(USkeletalMesh* SkeletalMesh, const FString PackageName, const int32 LODIndex)
 {
 	if (!SkeletalMesh || PackageName.IsEmpty())
@@ -604,8 +736,11 @@ UStaticMesh* UAnimToTextureBPLibrary::ConvertSkeletalMeshToStaticMesh(USkeletalM
 	}
 
 	// Create Temp Actor
+	check(GEditor);
 	UWorld* World = GEditor->GetEditorWorldContext().World();
+	check(World);
 	AActor* Actor = World->SpawnActor<AActor>();
+	check(Actor);
 
 	// Create Temp SkeletalMesh Component
 	USkeletalMeshComponent* MeshComponent = NewObject<USkeletalMeshComponent>(Actor);
@@ -701,30 +836,30 @@ UStaticMesh* UAnimToTextureBPLibrary::ConvertSkeletalMeshToStaticMesh(USkeletalM
 
 void UAnimToTextureBPLibrary::NormalizeVertexData(
 	const TArray<FVector3f>& Deltas, const TArray<FVector3f>& Normals,
-	FVector& MinBBox, FVector& SizeBBox,
-	TArray<FVector3f>& NormalizedDeltas, TArray<FVector3f>& NormalizedNormals)
+	FVector& OutMinBBox, FVector& OutSizeBBox,
+	TArray<FVector3f>& OutNormalizedDeltas, TArray<FVector3f>& OutNormalizedNormals)
 {
 	check(Deltas.Num() == Normals.Num());
 
 	// ---------------------------------------------------------------------------
 	// Compute Bounding Box
 	//
-	MinBBox = { TNumericLimits<float>::Max(), TNumericLimits<float>::Max(), TNumericLimits<float>::Max() };
+	OutMinBBox = { TNumericLimits<float>::Max(), TNumericLimits<float>::Max(), TNumericLimits<float>::Max() };
 	FVector3f MaxBBox = { TNumericLimits<float>::Min(), TNumericLimits<float>::Min(), TNumericLimits<float>::Min() };
 	
 	for (const FVector3f& Delta: Deltas)
 	{
 		// Find Min/Max BoundingBox
-		MinBBox.X = FMath::Min(Delta.X, MinBBox.X);
-		MinBBox.Y = FMath::Min(Delta.Y, MinBBox.Y);
-		MinBBox.Z = FMath::Min(Delta.Z, MinBBox.Z);
+		OutMinBBox.X = FMath::Min(Delta.X, OutMinBBox.X);
+		OutMinBBox.Y = FMath::Min(Delta.Y, OutMinBBox.Y);
+		OutMinBBox.Z = FMath::Min(Delta.Z, OutMinBBox.Z);
 
 		MaxBBox.X = FMath::Max(Delta.X, MaxBBox.X);
 		MaxBBox.Y = FMath::Max(Delta.Y, MaxBBox.Y);
 		MaxBBox.Z = FMath::Max(Delta.Z, MaxBBox.Z);
 	}
 
-	SizeBBox = (FVector)MaxBBox - MinBBox;
+	OutSizeBBox = (FVector)MaxBBox - OutMinBBox;
 
 	// ---------------------------------------------------------------------------
 	// Normalize Vertex Position Deltas
@@ -732,54 +867,54 @@ void UAnimToTextureBPLibrary::NormalizeVertexData(
 	
 	// Compute Normalization Factor per-axis.
 	const FVector NormFactor = {
-		1.0f / static_cast<float>(SizeBBox.X),
-		1.0f / static_cast<float>(SizeBBox.Y),
-		1.0f / static_cast<float>(SizeBBox.Z) };
+		1.f / static_cast<float>(OutSizeBBox.X),
+		1.f / static_cast<float>(OutSizeBBox.Y),
+		1.f / static_cast<float>(OutSizeBBox.Z) };
 
-	NormalizedDeltas.SetNumUninitialized(Deltas.Num());
+	OutNormalizedDeltas.SetNumUninitialized(Deltas.Num());
 	for (int32 Index = 0; Index < Deltas.Num(); ++Index)
 	{
-		NormalizedDeltas[Index] = (FVector3f)(((FVector)Deltas[Index] - MinBBox) * NormFactor);
+		OutNormalizedDeltas[Index] = (FVector3f)(((FVector)Deltas[Index] - OutMinBBox) * NormFactor);
 	}
 
 	// ---------------------------------------------------------------------------
 	// Normalize Vertex Normals
 	// And move them to [0, 1]
 	
-	NormalizedNormals.SetNumUninitialized(Normals.Num());
+	OutNormalizedNormals.SetNumUninitialized(Normals.Num());
 	for (int32 Index = 0; Index < Normals.Num(); ++Index)
 	{
-		NormalizedNormals[Index] = (Normals[Index].GetSafeNormal() + FVector3f::OneVector) * 0.5f;
+		OutNormalizedNormals[Index] = (Normals[Index].GetSafeNormal() + FVector3f::OneVector) * 0.5f;
 	}
 
 }
 
 void UAnimToTextureBPLibrary::NormalizeBoneData(
 	const TArray<FVector3f>& Positions, const TArray<FVector4>& Rotations,
-	FVector& MinBBox, FVector& SizeBBox, 
-	TArray<FVector3f>& NormalizedPositions, TArray<FVector4>& NormalizedRotations)
+	FVector& OutMinBBox, FVector& OutSizeBBox, 
+	TArray<FVector3f>& OutNormalizedPositions, TArray<FVector4>& OutNormalizedRotations)
 {
 	check(Positions.Num() == Rotations.Num());
 
 	// ---------------------------------------------------------------------------
 	// Compute Position Bounding Box
 	//
-	MinBBox = { TNumericLimits<float>::Max(), TNumericLimits<float>::Max(), TNumericLimits<float>::Max() };
+	OutMinBBox = { TNumericLimits<float>::Max(), TNumericLimits<float>::Max(), TNumericLimits<float>::Max() };
 	FVector3f MaxBBox = { TNumericLimits<float>::Min(), TNumericLimits<float>::Min(), TNumericLimits<float>::Min() };
 
 	for (const FVector3f& Position : Positions)
 	{
 		// Find Min/Max BoundingBox
-		MinBBox.X = FMath::Min(Position.X, MinBBox.X);
-		MinBBox.Y = FMath::Min(Position.Y, MinBBox.Y);
-		MinBBox.Z = FMath::Min(Position.Z, MinBBox.Z);
+		OutMinBBox.X = FMath::Min(Position.X, OutMinBBox.X);
+		OutMinBBox.Y = FMath::Min(Position.Y, OutMinBBox.Y);
+		OutMinBBox.Z = FMath::Min(Position.Z, OutMinBBox.Z);
 
 		MaxBBox.X = FMath::Max(Position.X, MaxBBox.X);
 		MaxBBox.Y = FMath::Max(Position.Y, MaxBBox.Y);
 		MaxBBox.Z = FMath::Max(Position.Z, MaxBBox.Z);
 	}
 
-	SizeBBox = (FVector)MaxBBox - MinBBox;
+	OutSizeBBox = (FVector)MaxBBox - OutMinBBox;
 
 	// ---------------------------------------------------------------------------
 	// Normalize Bone Position.
@@ -787,27 +922,27 @@ void UAnimToTextureBPLibrary::NormalizeBoneData(
 
 	// Compute Normalization Factor per-axis.
 	const FVector NormFactor = {
-		1.0f / static_cast<float>(SizeBBox.X),
-		1.0f / static_cast<float>(SizeBBox.Y),
-		1.0f / static_cast<float>(SizeBBox.Z) };
+		1.f / static_cast<float>(OutSizeBBox.X),
+		1.f / static_cast<float>(OutSizeBBox.Y),
+		1.f / static_cast<float>(OutSizeBBox.Z) };
 
-	NormalizedPositions.SetNumUninitialized(Positions.Num());
+	OutNormalizedPositions.SetNumUninitialized(Positions.Num());
 	for (int32 Index = 0; Index < Positions.Num(); ++Index)
 	{
-		NormalizedPositions[Index] = FVector3f(((FVector)Positions[Index] - MinBBox) * NormFactor);
+		OutNormalizedPositions[Index] = FVector3f(((FVector)Positions[Index] - OutMinBBox) * NormFactor);
 	}
 
 	// ---------------------------------------------------------------------------
 	// Normalize Rotations
 	// And move them to [0, 1]
-	NormalizedRotations.SetNumUninitialized(Rotations.Num());
+	OutNormalizedRotations.SetNumUninitialized(Rotations.Num());
 	for (int32 Index = 0; Index < Rotations.Num(); ++Index)
 	{
 		const FVector4 Axis = Rotations[Index];
 		const float Angle = Rotations[Index].W; // Angle are returned in radians and they go from [0-pi*2]
 
-		NormalizedRotations[Index] = (Axis.GetSafeNormal() + FVector::OneVector) * 0.5f;
-		NormalizedRotations[Index].W = Angle / (PI * 2.f);
+		OutNormalizedRotations[Index] = (Axis.GetSafeNormal() + FVector::OneVector) * 0.5f;
+		OutNormalizedRotations[Index].W = Angle / (PI * 2.f);
 	}
 }
 
@@ -816,19 +951,10 @@ bool UAnimToTextureBPLibrary::CreateUVChannel(
 	UStaticMesh* StaticMesh, const int32 LODIndex, const int32 UVChannelIndex,
 	const int32 Height, const int32 Width)
 {
-	if (!StaticMesh)
-	{
-		return false;
-	}
+	check(StaticMesh);
 
-	// Get Render Data
-	const FStaticMeshRenderData* RenderData = StaticMesh->GetRenderData();
-	check(RenderData);
-
-	// Get LOD Data
-	if (!RenderData->LODResources.IsValidIndex(LODIndex))
+	if (!StaticMesh->IsSourceModelValid(LODIndex))
 	{
-		UE_LOG(LogTemp, Warning, TEXT("Invalid LODIndex: %i"), LODIndex);
 		return false;
 	}
 
@@ -838,20 +964,7 @@ bool UAnimToTextureBPLibrary::CreateUVChannel(
 	FMeshDescription* MeshDescription = StaticMesh->GetMeshDescription(LODIndex);
 	check(MeshDescription);
 
-	// Check if UVChannel is being used by the Lightmap UV
-	if (StaticMesh->IsSourceModelValid(LODIndex))
-	{
-		const FStaticMeshSourceModel& SourceModel = StaticMesh->GetSourceModel(LODIndex);
-		if (SourceModel.BuildSettings.bGenerateLightmapUVs &&
-			SourceModel.BuildSettings.DstLightmapIndex == UVChannelIndex)
-		{
-			UE_LOG(LogTemp, Warning, TEXT("Invalid UVChannel: %i. Already used by LightMap"), UVChannelIndex);
-			return false;
-		}
-	}
-
 	// Add New UVChannel.
-	// UE_LOG(LogTemp, Warning, TEXT("UVChannel: %i. Number of existing UVChannels: %i"), UVChannelIndex, StaticMesh->GetNumUVChannels(LODIndex));
 	if (UVChannelIndex == StaticMesh->GetNumUVChannels(LODIndex))
 	{
 		if (!StaticMesh->InsertUVChannel(LODIndex, UVChannelIndex))
@@ -884,11 +997,9 @@ bool UAnimToTextureBPLibrary::CreateUVChannel(
 	// Set Full Precision UVs
 	SetFullPrecisionUVs(StaticMesh, LODIndex, true);
 
+	// Set UVs
 	if (StaticMesh->SetUVChannel(LODIndex, UVChannelIndex, TexCoords))
-	{		
-		// Update and Mark to Save.
-		StaticMesh->MarkPackageDirty();
-
+	{
 		return true;
 	}
 	else
@@ -934,6 +1045,8 @@ bool UAnimToTextureBPLibrary::FindBestResolution(
 
 void UAnimToTextureBPLibrary::SetFullPrecisionUVs(UStaticMesh* StaticMesh, int32 LODIndex, bool bFullPrecision)
 {
+	check(StaticMesh);
+
 	if (StaticMesh->IsSourceModelValid(LODIndex))
 	{
 		FStaticMeshSourceModel& SourceModel = StaticMesh->GetSourceModel(LODIndex);
@@ -943,44 +1056,41 @@ void UAnimToTextureBPLibrary::SetFullPrecisionUVs(UStaticMesh* StaticMesh, int32
 
 void UAnimToTextureBPLibrary::SetBoundsExtensions(UStaticMesh* StaticMesh, const FVector& MinBBox, const FVector& SizeBBox)
 {
+	check(StaticMesh);
+
 	// Calculate MaxBBox
 	const FVector MaxBBox = SizeBBox + MinBBox;
 
-	if (IsValid(StaticMesh))
-	{
-		// Reset current extension bounds
-		const FVector PositiveBoundsExtension = StaticMesh->GetPositiveBoundsExtension();
-		const FVector NegativeBoundsExtension = StaticMesh->GetNegativeBoundsExtension();
+	// Reset current extension bounds
+	const FVector PositiveBoundsExtension = StaticMesh->GetPositiveBoundsExtension();
+	const FVector NegativeBoundsExtension = StaticMesh->GetNegativeBoundsExtension();
 		
-		// Get current BoundingBox including extensions
-		FBox BoundingBox = StaticMesh->GetBoundingBox();
+	// Get current BoundingBox including extensions
+	FBox BoundingBox = StaticMesh->GetBoundingBox();
 		
-		// Remove extensions from BoundingBox
-		BoundingBox.Max -= PositiveBoundsExtension;
-		BoundingBox.Min += NegativeBoundsExtension;
+	// Remove extensions from BoundingBox
+	BoundingBox.Max -= PositiveBoundsExtension;
+	BoundingBox.Min += NegativeBoundsExtension;
 		
-		// Calculate New BoundingBox
-		FVector NewMaxBBox(
-			FMath::Max(BoundingBox.Max.X, MaxBBox.X),
-			FMath::Max(BoundingBox.Max.Y, MaxBBox.Y),
-			FMath::Max(BoundingBox.Max.Z, MaxBBox.Z)
-		);
+	// Calculate New BoundingBox
+	FVector NewMaxBBox(
+		FMath::Max(BoundingBox.Max.X, MaxBBox.X),
+		FMath::Max(BoundingBox.Max.Y, MaxBBox.Y),
+		FMath::Max(BoundingBox.Max.Z, MaxBBox.Z)
+	);
 		
-		FVector NewMinBBox(
-			FMath::Min(BoundingBox.Min.X, MinBBox.X),
-			FMath::Min(BoundingBox.Min.Y, MinBBox.Y),
-			FMath::Min(BoundingBox.Min.Z, MinBBox.Z)
-		);
-
-		// Calculate New Extensions
-		FVector NewPositiveBoundsExtension = NewMaxBBox - BoundingBox.Max;
-		FVector NewNegativeBoundsExtension = BoundingBox.Min - NewMinBBox;
+	FVector NewMinBBox(
+		FMath::Min(BoundingBox.Min.X, MinBBox.X),
+		FMath::Min(BoundingBox.Min.Y, MinBBox.Y),
+		FMath::Min(BoundingBox.Min.Z, MinBBox.Z)
+	);
+
+	// Calculate New Extensions
+	FVector NewPositiveBoundsExtension = NewMaxBBox - BoundingBox.Max;
+	FVector NewNegativeBoundsExtension = BoundingBox.Min - NewMinBBox;
 				
-		// Update StaticMesh
-		StaticMesh->SetPositiveBoundsExtension(NewPositiveBoundsExtension);
-		StaticMesh->SetNegativeBoundsExtension(NewNegativeBoundsExtension);
-		StaticMesh->CalculateExtendedBounds();
-		StaticMesh->PostEditChange();		
-		StaticMesh->MarkPackageDirty();
-	}
+	// Update StaticMesh
+	StaticMesh->SetPositiveBoundsExtension(NewPositiveBoundsExtension);
+	StaticMesh->SetNegativeBoundsExtension(NewNegativeBoundsExtension);
+	StaticMesh->CalculateExtendedBounds();
 }
diff --git a/Engine/Plugins/Animation/AnimToTexture/Source/AnimToTextureEditor/Private/AnimToTextureMeshMapping.cpp b/Engine/Plugins/Animation/AnimToTexture/Source/AnimToTextureEditor/Private/AnimToTextureMeshMapping.cpp
new file mode 100644
--- /dev/null
+++ b/Engine/Plugins/Animation/AnimToTexture/Source/AnimToTextureEditor/Private/AnimToTextureMeshMapping.cpp
@@ -0,0 +1,301 @@
+// Copyright Epic Games, Inc. All Rights Reserved.
+#include "AnimToTextureMeshMapping.h"
+
+namespace AnimToTexture_Private
+{
+
+FVector3f FVertexToMeshMapping::FindClosestPointToTriangle(const FVector3f& P, const FVector3f& A, const FVector3f& B, const FVector3f& C)
+{
+	const FVector3f AB = B - A;
+	const FVector3f AC = C - A;
+	const FVector3f AP = P - A;
+
+	const float D1 = FVector3f::DotProduct(AB, AP);
+	const float D2 = FVector3f::DotProduct(AC, AP);
+	if (D1 <= 0.f && D2 <= 0.f)
+	{
+		return A;
+	}
+
+	const FVector3f BP = P - B;
+	const float D3 = FVector3f::DotProduct(AB, BP);
+	const float D4 = FVector3f::DotProduct(AC, BP);
+	if (D3 >= 0.f && D4 <= D3)
+	{
+		return B;
+	}
+
+	const FVector3f CP = P - C;
+	const float D5 = FVector3f::DotProduct(AB, CP);
+	const float D6 = FVector3f::DotProduct(AC, CP);
+	if (D6 >= 0.f && D5 <= D6)
+	{
+		return C;
+	}
+
+	const float VC = D1 * D4 - D3 * D2;
+	if (VC <= 0.f && D1 >= 0.f && D3 <= 0.f)
+	{
+		const float V = D1 / (D1 - D3);
+		return A + V * AB;
+	}
+
+	const float VB = D5 * D2 - D1 * D6;
+	if (VB <= 0.f && D2 >= 0.f && D6 <= 0.f)
+	{
+		const float V = D2 / (D2 - D6);
+		return A + V * AC;
+	}
+
+	const float VA = D3 * D6 - D5 * D4;
+	if (VA <= 0.f && (D4 - D3) >= 0.f && (D5 - D6) >= 0.f)
+	{
+		const float V = (D4 - D3) / ((D4 - D3) + (D5 - D6));
+		return B + V * (C - B);
+	}
+
+	const float Denom = 1.0f / (VA + VB + VC);
+	const float V = VB * Denom;
+	const float W = VC * Denom;
+
+	return A + V * AB + W * AC;
+}
+
+FVector3f FVertexToMeshMapping::GetTriangleNormal(const FVector3f& A, const FVector3f& B, const FVector3f& C)
+{
+	const FVector3f V0 = B - A;
+	const FVector3f V1 = C - A;
+	return FVector3f::CrossProduct(V0, V1).GetSafeNormal();
+}
+
+FMatrix44f FVertexToMeshMapping::GetTriangleMatrix(const FVector3f& A, const FVector3f& B, const FVector3f& C)
+{	
+	const FVector3f V0 = (B - A).GetSafeNormal();
+	const FVector3f V1 = GetTriangleNormal(A, B, C);
+	const FVector3f V2 = FVector3f::CrossProduct(V0, V1).GetSafeNormal();
+	
+	FMatrix44f Matrix = FMatrix44f::Identity; // SetAxes doesn't set [3] elements.
+	Matrix.SetAxes(&V0, &V1, &V2, &A);
+
+	return Matrix;
+}
+
+FVector3f FVertexToMeshMapping::BarycentricCoordinates(const FVector3f& P, const FVector3f& A, const FVector3f& B, const FVector3f& C)
+{
+	const FVector3f V0 = B - A;
+	const FVector3f V1 = C - A;
+	const FVector3f V2 = P - A;
+	const float D00 = FVector3f::DotProduct(V0, V0);
+	const float D01 = FVector3f::DotProduct(V0, V1);
+	const float D11 = FVector3f::DotProduct(V1, V1);
+	const float D20 = FVector3f::DotProduct(V2, V0);
+	const float D21 = FVector3f::DotProduct(V2, V1);
+	const float Denom = 1.0f / (D00 * D11 - D01 * D01);
+	
+	const float V = (D11 * D20 - D01 * D21) * Denom;
+	const float W = (D00 * D21 - D01 * D20) * Denom;
+	const float U = 1.0f - V - W;
+
+	return FVector3f(U, V, W);
+}
+
+int32 FVertexToMeshMapping::FindClosestTriangle(const FVector3f& Point, const TArray<FVector3f>& Vertices, const TArray<FIntVector3>& Triangles,
+	FVector3f& OutClosestPoint, FVector3f& OutBarycentricCoords)
+{
+	int32 ClosestTriangle = INDEX_NONE;
+	float ClosestDistance = TNumericLimits<float>::Max();
+
+	for (int32 TriangleIndex = 0; TriangleIndex < Triangles.Num(); TriangleIndex++)
+	{
+		const FIntVector3& Triangle = Triangles[TriangleIndex];
+
+		// Check if valid indices
+		if (Vertices.IsValidIndex(Triangle.X) && 
+			Vertices.IsValidIndex(Triangle.Y) && 
+			Vertices.IsValidIndex(Triangle.Z))
+		{
+			const FVector3f& A = Vertices[Triangle.X];
+			const FVector3f& B = Vertices[Triangle.Y];
+			const FVector3f& C = Vertices[Triangle.Z];
+
+			const FVector3f ClosestPoint = FindClosestPointToTriangle(Point, A, B, C);
+			const float Distance = FVector3f::Distance(Point, ClosestPoint);
+			if (Distance < ClosestDistance)
+			{
+				OutClosestPoint = ClosestPoint;
+				ClosestTriangle = TriangleIndex;
+				ClosestDistance = Distance;
+			}
+		}
+	}
+	check(ClosestTriangle != INDEX_NONE);
+
+	// Compute Barycentric Coordinates
+	const FIntVector3& Triangle = Triangles[ClosestTriangle];
+	const FVector3f& A = Vertices[Triangle.X];
+	const FVector3f& B = Vertices[Triangle.Y];
+	const FVector3f& C = Vertices[Triangle.Z];
+	OutBarycentricCoords = BarycentricCoordinates(OutClosestPoint, A, B, C);
+
+	return ClosestTriangle;
+}
+
+//int32 FindClosestVertex(const FVector3f& Point, const TArray<FVector3f>& Vertices)
+//{
+//	float MinDistance = TNumericLimits<float>::Max();
+//	int32 ClosestIndex = INDEX_NONE;
+//
+//	// Find Closest SkeletalMesh Vertex.
+//	for (int32 VertexIndex = 0; VertexIndex < Vertices.Num(); ++VertexIndex)
+//	{
+//		const float Distance = FVector3f::Dist(Point, Vertices[VertexIndex]);
+//		if (Distance < MinDistance)
+//		{
+//			MinDistance = Distance;
+//			ClosestIndex = VertexIndex;
+//		}
+//	}
+//
+//	return ClosestIndex;
+//}
+
+
+void FVertexToMeshMapping::Create(
+	const UStaticMesh* StaticMesh, const int32 StaticMeshLODIndex,	
+	const USkeletalMesh* SkeletalMesh, const int32 SkeletalMeshLODIndex,
+	TArray<FVertexToMeshMapping>& OutMapping)
+{
+	check(StaticMesh);
+	check(SkeletalMesh);
+	OutMapping.Reset();
+
+	// Get StaticMesh Vertices
+	TArray<FVector3f> StaticMeshVertices;
+	TArray<FVector3f> StaticMeshNormals;
+	const int32 NumStaticMeshVertices = GetVertices(StaticMesh, StaticMeshLODIndex, StaticMeshVertices, StaticMeshNormals);
+
+	// Get SkeletalMesh Vertices
+	TArray<FVector3f> SkeletalMeshVertices;
+	const int32 NumSkeletalMeshVertices = GetVertices(SkeletalMesh, SkeletalMeshLODIndex, SkeletalMeshVertices);
+
+	// Get SkeletalMesh Triangles
+	TArray<FIntVector3> SkeletalMeshTriangles;
+	GetTriangles(SkeletalMesh, SkeletalMeshLODIndex, SkeletalMeshTriangles);
+
+	// Allocate
+	OutMapping.SetNumUninitialized(NumStaticMeshVertices);
+
+	// Fill Out Mapping
+	for (int32 VertexIndex = 0; VertexIndex < NumStaticMeshVertices; VertexIndex++)
+	{
+		const FVector3f& StaticMeshVertex = StaticMeshVertices[VertexIndex];
+
+		FVector3f ClosestPoint;
+		FVector3f BarycentricCoords;
+
+		// Find Closest Triangle (on SkeletalMesh)
+		const int32 ClosestTriangle = FindClosestTriangle(StaticMeshVertex, SkeletalMeshVertices, SkeletalMeshTriangles, ClosestPoint, BarycentricCoords);
+		check(SkeletalMeshTriangles.IsValidIndex(ClosestTriangle));
+		
+		// Closest Triangle Vertices
+		const FIntVector3& Triangle = SkeletalMeshTriangles[ClosestTriangle];
+		const FVector3f& A = SkeletalMeshVertices[Triangle.X];
+		const FVector3f& B = SkeletalMeshVertices[Triangle.Y];
+		const FVector3f& C = SkeletalMeshVertices[Triangle.Z];
+		
+		// Compute Triangle Matrix
+		const FMatrix44f Matrix = GetTriangleMatrix(A, B, C);
+		
+		// Store Mapping 
+		OutMapping[VertexIndex].Triangle = Triangle;
+		OutMapping[VertexIndex].BarycentricCoords = BarycentricCoords;
+		OutMapping[VertexIndex].InvMatrix = Matrix.Inverse();
+	}
+}
+
+
+void FVertexToMeshMapping::InterpolateVertexSkinWeights(const VertexSkinWeightMax& A, const VertexSkinWeightMax& B, const VertexSkinWeightMax& C,
+	VertexSkinWeightMax& OutWeights) const
+{	
+	// NOTE: this could use NumBones and remove the 256 limit.
+	TStaticArray<float, 256> DenseA(InPlace, 0.f);
+	TStaticArray<float, 256> DenseB(InPlace, 0.f);
+	TStaticArray<float, 256> DenseC(InPlace, 0.f);
+
+	// Convert to Dense Arrays	
+	for (int32 Index = 0; Index < MAX_TOTAL_INFLUENCES; Index++)
+	{
+		if (A.BoneWeights[Index])
+		{
+			DenseA[A.MeshBoneIndices[Index]] = (float)A.BoneWeights[Index] / 255.f;
+		}
+		if (B.BoneWeights[Index])
+		{
+			DenseB[B.MeshBoneIndices[Index]] = (float)B.BoneWeights[Index] / 255.f;
+		}
+		if (C.BoneWeights[Index])
+		{
+			DenseC[C.MeshBoneIndices[Index]] = (float)C.BoneWeights[Index] / 255.f;
+		}
+	}
+
+	// Blend Dense Arrays
+	TArray<TPair<uint8, uint16>> SortedVertexSkinWeights;
+	SortedVertexSkinWeights.SetNumUninitialized(256);
+
+	for (int32 Index = 0; Index < 256; Index++)
+	{
+		const float BoneWeightA = DenseA[Index] * BarycentricCoords.X;
+		const float BoneWeightB = DenseB[Index] * BarycentricCoords.Y;
+		const float BoneWeightC = DenseC[Index] * BarycentricCoords.Z;
+
+		const uint8 BoneWeight = (uint8)FMath::RoundToInt((BoneWeightA + BoneWeightB + BoneWeightC) * 255.f);
+		SortedVertexSkinWeights[Index] = TPair<uint8, uint16>(BoneWeight, Index);
+	}
+
+	// Sort Weights and return Sparse
+	SortedVertexSkinWeights.Sort();
+	for (int32 Index = 0; Index < MAX_TOTAL_INFLUENCES; Index++)
+	{
+		OutWeights.BoneWeights[Index]     = SortedVertexSkinWeights[255 - Index].Key;
+		OutWeights.MeshBoneIndices[Index] = SortedVertexSkinWeights[255 - Index].Value;
+	}
+}
+
+void FVertexToMeshMapping::InterpolateSkinWeights(const TArray<FVertexToMeshMapping>& Mapping, const TArray<VertexSkinWeightMax>& SkinWeights,
+	TArray<VertexSkinWeightMax>& OutSkinWeights)
+{
+	// Allocate StaticMesh Weights
+	OutSkinWeights.SetNumUninitialized(Mapping.Num());
+
+	for (int32 VertexIndex = 0; VertexIndex < Mapping.Num(); VertexIndex++)
+	{
+		const FVertexToMeshMapping& VertexMapping = Mapping[VertexIndex];
+		const FIntVector3& Triangle = VertexMapping.Triangle;
+
+		if (SkinWeights.IsValidIndex(Triangle.X) &&
+			SkinWeights.IsValidIndex(Triangle.Y) &&
+			SkinWeights.IsValidIndex(Triangle.Z))
+		{
+			const VertexSkinWeightMax& A = SkinWeights[Triangle.X];
+			const VertexSkinWeightMax& B = SkinWeights[Triangle.Y];
+			const VertexSkinWeightMax& C = SkinWeights[Triangle.Z];
+
+			Mapping[VertexIndex].InterpolateVertexSkinWeights(A, B, C, OutSkinWeights[VertexIndex]);
+		}
+	}
+}
+
+FVector3f FVertexToMeshMapping::TransformPosition(const FVector3f& P, const FVector3f& A, const FVector3f& B, const FVector3f& C) const
+{	
+	const FMatrix44f Matrix = GetTriangleMatrix(A, B, C);	
+	return Matrix.TransformPosition(InvMatrix.TransformPosition(P));
+}
+
+FVector3f FVertexToMeshMapping::TransformVector(const FVector3f& V, const FVector3f& A, const FVector3f& B, const FVector3f& C) const
+{
+	const FMatrix44f Matrix = GetTriangleMatrix(A, B, C);
+	return Matrix.TransformVector(InvMatrix.TransformVector(V));
+}
+
+} // end namespace AnimToTexture_Private
diff --git a/Engine/Plugins/Animation/AnimToTexture/Source/AnimToTextureEditor/Private/AnimToTextureSkeletalMesh.cpp b/Engine/Plugins/Animation/AnimToTexture/Source/AnimToTextureEditor/Private/AnimToTextureSkeletalMesh.cpp
--- a/Engine/Plugins/Animation/AnimToTexture/Source/AnimToTextureEditor/Private/AnimToTextureSkeletalMesh.cpp
+++ b/Engine/Plugins/Animation/AnimToTexture/Source/AnimToTextureEditor/Private/AnimToTextureSkeletalMesh.cpp
@@ -1,23 +1,63 @@
 // Copyright Epic Games, Inc. All Rights Reserved.
 #include "AnimToTextureSkeletalMesh.h"
-#include "Rendering/SkeletalMeshModel.h"
 #include "Rendering/SkeletalMeshRenderData.h"
 #include "Engine/StaticMesh.h"
 #include "Engine/SkeletalMesh.h"
 #include "Components/SkeletalMeshComponent.h"
 #include "MeshDescription.h"
-#include "BoneIndices.h"
+#include "StaticMeshAttributes.h"
 
 namespace AnimToTexture_Private
 {
 
-int32 GetVertices(const USkeletalMesh& Mesh, const int32 LODIndex, 
+int32 GetVertices(const UStaticMesh* StaticMesh, const int32 LODIndex,
 	TArray<FVector3f>& OutPositions, TArray<FVector3f>& OutNormals)
 {
+	check(StaticMesh);
 	OutPositions.Reset();
 	OutNormals.Reset();
+
+	if (!StaticMesh->IsMeshDescriptionValid(LODIndex))
+	{
+		return INDEX_NONE;
+	}
+
+	// Get Mesh Description
+	const FMeshDescription* MeshDescription = StaticMesh->GetMeshDescription(LODIndex);
+	check(MeshDescription);
+	
+	// Get MeshDescription Vertices
+	const FVertexArray& Vertices = MeshDescription->Vertices();
+	const int32 NumVertices = Vertices.Num();
+	OutPositions.SetNum(NumVertices);
+	OutNormals.SetNumZeroed(NumVertices);
+	
+	for (const FVertexID& VertexID : Vertices.GetElementIDs())
+	{
+		const int32 VertexIndex = VertexID.GetValue();
+
+		// Get Vertex Position
+		OutPositions[VertexIndex] = MeshDescription->GetVertexPosition(VertexID);
+
+		// Get Avg Vertex Normals
+		const TArrayView<const FVertexInstanceID> VertexInstances = MeshDescription->GetVertexVertexInstanceIDs(VertexID);
+		for (const FVertexInstanceID VertexInstanceID : VertexInstances)
+		{
+			const FVector3f VertexInstanceNormal = MeshDescription->VertexInstanceAttributes().GetAttribute<FVector3f>(VertexInstanceID, MeshAttribute::VertexInstance::Normal);
+			OutNormals[VertexIndex] += VertexInstanceNormal / VertexInstances.Num();
+		}
+	}
+	
+	return NumVertices;
+}
+
+int32 GetVertices(const USkeletalMesh* SkeletalMesh, const int32 LODIndex, 
+	TArray<FVector3f>& OutPositions)
+{
+	check(SkeletalMesh);
+	OutPositions.Reset();
 	
-	const FSkeletalMeshRenderData* RenderData = Mesh.GetResourceForRendering();
+	const FSkeletalMeshRenderData* RenderData = SkeletalMesh->GetResourceForRendering();
 	check(RenderData);
 
 	if (!RenderData->LODRenderData.IsValidIndex(LODIndex))
@@ -31,73 +71,50 @@ int32 GetVertices(const USkeletalMesh& Mesh, const int32 LODIndex,
 	// Get Total Num of Vertices (for all sections)
 	const int32 NumVertices = LODRenderData.GetNumVertices();
 	OutPositions.SetNumUninitialized(NumVertices);
-	OutNormals.SetNumUninitialized(NumVertices);
 	
 	for (int32 VertexIndex = 0; VertexIndex < NumVertices; ++VertexIndex)
 	{
 		OutPositions[VertexIndex] = LODRenderData.StaticVertexBuffers.PositionVertexBuffer.VertexPosition(VertexIndex);
-		OutNormals[VertexIndex] = LODRenderData.StaticVertexBuffers.StaticMeshVertexBuffer.VertexTangentZ(VertexIndex);
 	};
 
 	return NumVertices;
 }
 
-int32 FindClosestVertex(const FVector3f& Point, const TArray<FVector3f>& Vertices)
+int32 GetTriangles(const USkeletalMesh* SkeletalMesh, const int32 LODIndex,
+	TArray<FIntVector3>& OutTriangles)
 {
-	float MinDistance = TNumericLimits<float>::Max();
-	int32 ClosestIndex = INDEX_NONE;
+	check(SkeletalMesh);
+	OutTriangles.Reset();
 
-	// Find Closest SkeletalMesh Vertex.
-	for (int32 VertexIndex = 0; VertexIndex < Vertices.Num(); ++VertexIndex)
+	const FSkeletalMeshRenderData* RenderData = SkeletalMesh->GetResourceForRendering();
+	check(RenderData);
+
+	if (!RenderData->LODRenderData.IsValidIndex(LODIndex))
 	{
-		const float Distance = FVector3f::Dist(Point, Vertices[VertexIndex]);
-		if (Distance < MinDistance)
-		{
-			MinDistance = Distance;
-			ClosestIndex = VertexIndex;
-		}
+		return INDEX_NONE;
 	}
 
-	return ClosestIndex;
-}
-
-void GetStaticToSkeletalMapping(
-	const UStaticMesh& StaticMesh, const int32 StaticLODIndex,
-	const USkeletalMesh& SkeletalMesh, const int32 SkeletalLODIndex,
-	TArray<FVector3f>& OutStaticVertices, TArray<FVector3f>& OutSkeletalVertices,
-	TArray<int32>& OutStaticToSkeletalMapping)
-{
-	// Reset
-	OutStaticVertices.Reset();
-	OutSkeletalVertices.Reset();
-	OutStaticToSkeletalMapping.Reset();
-
-	// Get Mesh Description
-	const FMeshDescription* MeshDescription = StaticMesh.GetMeshDescription(StaticLODIndex);
-	check(MeshDescription);
+	// Get LOD Data
+	const FSkeletalMeshLODRenderData& LODRenderData = RenderData->LODRenderData[LODIndex];
 
-	// Get SkeletalMesh Vertices
-	TArray<FVector3f> Normals;
-	const int32 SkelNumVertices = AnimToTexture_Private::GetVertices(SkeletalMesh, SkeletalLODIndex, OutSkeletalVertices, Normals);
+	// Get Indices
+	const FMultiSizeIndexContainer& IndexBuffer = LODRenderData.MultiSizeIndexContainer;
+	TArray<uint32> IndexArray;
+	IndexBuffer.GetIndexBuffer(IndexArray);
 
-	// Get MeshDescription Vertices
-	const FVertexArray& Vertices = MeshDescription->Vertices();
-	const int32 NumVertices = Vertices.Num();
-	OutStaticVertices.SetNum(NumVertices);
-	OutStaticToSkeletalMapping.SetNum(NumVertices);
+	// Allocate Triangles
+	const uint32 NumTriangles = IndexArray.Num() / 3;
+	OutTriangles.SetNumUninitialized(NumTriangles);
 
-	for (const FVertexID VertexID : Vertices.GetElementIDs())
+	for (uint32 TriangleIndex = 0; TriangleIndex < NumTriangles; TriangleIndex++)
 	{
-		const int32 VertexIndex = VertexID.GetValue();
-
-		// Get Mesh Description Vertex Position
-		const FVector3f Vertex = MeshDescription->GetVertexPosition(VertexID);
-		OutStaticVertices[VertexIndex] = Vertex;
-
-		// Find Closest SkeletalMesh Vertex
-		OutStaticToSkeletalMapping[VertexIndex] = FindClosestVertex(Vertex, OutSkeletalVertices);
+		OutTriangles[TriangleIndex].X = IndexArray[TriangleIndex * 3];
+		OutTriangles[TriangleIndex].Y = IndexArray[TriangleIndex * 3 + 1];
+		OutTriangles[TriangleIndex].Z = IndexArray[TriangleIndex * 3 + 2];
 	}
-}
+
+	return NumTriangles;
+};
 
 int32 GetNumBones(const USkeletalMesh* SkeletalMesh)
 {
@@ -178,124 +195,14 @@ void GetBoneNames(const USkeletalMesh* SkeletalMesh, TArray<FName>& OutNames)
 	}
 };
 
-void GetSkinnedVertices(USkeletalMeshComponent* MeshComponent, const int32 LODIndex, 
-	TArray<FVector3f>& OutPositions, TArray<FVector3f>& OutNormals)
+void GetSkinWeights(const USkeletalMesh* SkeletalMesh, const int32 LODIndex, 
+	TArray<VertexSkinWeightMax>& OutSkinWeights)
 {
-	OutPositions.Reset();
-	OutNormals.Reset();
-
-	if (!MeshComponent)
-	{
-		return;
-	}
-
-	// Get SkeletalMesh
-	USkeletalMesh* SkeletalMesh = MeshComponent->GetSkeletalMeshAsset();
-	if (!SkeletalMesh)
-	{
-		return;
-	}
-
-	// Get Render Data 
-	FSkeletalMeshRenderData* RenderData = SkeletalMesh->GetResourceForRendering();
-	if (!RenderData->LODRenderData.IsValidIndex(LODIndex))
-	{
-		return;
-	};
-
-	// Get Matrices
-	TArray<FMatrix44f> RefToLocals;
-	MeshComponent->CacheRefToLocalMatrices(RefToLocals);
-
-	// ---------------------------------------------------------------------------------------------
-
-	// Get Ref-Pose Vertices
-	TArray<FVector3f> Vertices;
-	TArray<FVector3f> Normals;
-	const int32 NumVertices = GetVertices(*SkeletalMesh, LODIndex, Vertices, Normals);
-
-	// TODO: Add Morph Deltas to Vertices.
-
-	// Get Weights
-	TArray<TVertexSkinWeight<MAX_TOTAL_INFLUENCES>> SkinWeights;
-	GetSkinWeightsData(*SkeletalMesh, LODIndex, SkinWeights);
-
-	OutPositions.SetNumUninitialized(NumVertices);
-	OutNormals.SetNumUninitialized(NumVertices);
-
-	for (int32 VertexIndex = 0; VertexIndex < NumVertices; ++VertexIndex)
-	{
-		const FVector3f& Vertex = Vertices[VertexIndex];
-		const FVector3f& Normal = Normals[VertexIndex];
-		const TVertexSkinWeight<MAX_TOTAL_INFLUENCES>& Weights = SkinWeights[VertexIndex];
-
-		FVector4f SkinnedVertex(0);
-		FVector4f SkinnedNormal(0);
-
-		for (int32 InfluenceIndex = 0; InfluenceIndex < MAX_TOTAL_INFLUENCES; InfluenceIndex++)
-		{
-			const uint8& BoneWeight = Weights.BoneWeights[InfluenceIndex];
-			const uint16& MeshBoneIndex = Weights.MeshBoneIndices[InfluenceIndex];
-
-			// Get Matrix
-			const FMatrix44f& RefToLocal = RefToLocals[MeshBoneIndex];
-
-			const float Weight = (float)BoneWeight / 255.f;
-			SkinnedVertex += RefToLocal.TransformPosition(Vertex) * Weight;
-			SkinnedNormal += RefToLocal.TransformVector(Normal) * Weight;
-		}
-
-		OutPositions[VertexIndex] = SkinnedVertex;
-		OutNormals[VertexIndex] = SkinnedNormal;
-	};
-	
-};
-
-void ReduceSkinWeightsData(const TArray<TVertexSkinWeight<MAX_TOTAL_INFLUENCES>>& InSkinWeights, TArray<TVertexSkinWeight<4>>& OutSkinWeights)
-{
-	OutSkinWeights.SetNumUninitialized(InSkinWeights.Num());
-
-	for (int32 VertexIndex = 0; VertexIndex < InSkinWeights.Num(); ++VertexIndex)
-	{
-		GetMaxFourInfluences(InSkinWeights[VertexIndex], OutSkinWeights[VertexIndex]);
-	}
-}
-
-void DecomposeTransformation(const FTransform& Transform, FVector3f& Translation, FVector4& Rotation)
-{
-	// Get Translation
-	Translation = (FVector3f)Transform.GetTranslation();
-
-	// Get Rotation 
-	const FQuat Quat = Transform.GetRotation();
-
-	FVector Axis;
-	float Angle;
-	Quat.ToAxisAndAngle(Axis, Angle);
-
-	Rotation = FVector4(Axis, Angle);
-}
-
-void DecomposeTransformations(const TArray<FTransform>& Transforms, TArray<FVector3f>& Translations, TArray<FVector4>& Rotations)
-{
-	const int32 NumTransforms = Transforms.Num();
-	Translations.SetNumUninitialized(NumTransforms);
-	Rotations.SetNumUninitialized(NumTransforms);
-
-	for (int32 Index = 0; Index < NumTransforms; ++Index)
-	{
-		DecomposeTransformation(Transforms[Index], Translations[Index], Rotations[Index]);
-	}
-};
-
-
-void GetSkinWeightsData(const USkeletalMesh& SkeletalMesh, const int32 LODIndex, TArray<TVertexSkinWeight<MAX_TOTAL_INFLUENCES>>& SkinWeights)
-{
-	// Reset Weights
-	SkinWeights.Reset();
+	check(SkeletalMesh);
+	OutSkinWeights.Reset();
 
 	// Get Render Data
-	const FSkeletalMeshRenderData* RenderData = SkeletalMesh.GetResourceForRendering();
+	const FSkeletalMeshRenderData* RenderData = SkeletalMesh->GetResourceForRendering();
 	check(RenderData);
 
 	if (!RenderData->LODRenderData.IsValidIndex(LODIndex))
@@ -317,10 +224,10 @@ void GetSkinWeightsData(const USkeletalMesh& SkeletalMesh, const int32 LODIndex,
 	SkinWeightVertexBuffer->GetSkinWeights(SkinWeightsInfo);
 
 	// Allocated SkinWeightData
-	SkinWeights.SetNumUninitialized(SkinWeightsInfo.Num());
+	OutSkinWeights.SetNumUninitialized(SkinWeightsInfo.Num());
 
 	// Loop thru vertices
-	for (int32 VertexIndex = 0; VertexIndex < SkinWeightsInfo.Num(); ++VertexIndex)
+	for (int32 VertexIndex = 0; VertexIndex < SkinWeightsInfo.Num(); VertexIndex++)
 	{
 		// Find Section From Global Vertex Index
 		// NOTE: BoneMap is stored by Section.
@@ -335,41 +242,88 @@ void GetSkinWeightsData(const USkeletalMesh& SkeletalMesh, const int32 LODIndex,
 		const FSkinWeightInfo& SkinWeightInfo = SkinWeightsInfo[VertexIndex];
 
 		// Store Weights
-		for (int32 InfluenceIndex = 0; InfluenceIndex < MAX_TOTAL_INFLUENCES; ++InfluenceIndex)
+		for (int32 Index = 0; Index < MAX_TOTAL_INFLUENCES; Index++)
 		{
-			const uint16& BoneWeight = SkinWeightInfo.InfluenceWeights[InfluenceIndex];
-			const uint16& BoneIndex = SkinWeightInfo.InfluenceBones[InfluenceIndex];
+			const uint8& BoneWeight = SkinWeightInfo.InfluenceWeights[Index];
+			const uint16& BoneIndex = SkinWeightInfo.InfluenceBones[Index];
 			const uint16& MeshBoneIndex = RenderSection.BoneMap[BoneIndex];
 
-			SkinWeights[VertexIndex].BoneWeights[InfluenceIndex] = BoneWeight >> 8;
-			SkinWeights[VertexIndex].BoneIndices[InfluenceIndex] = BoneIndex;
-			SkinWeights[VertexIndex].MeshBoneIndices[InfluenceIndex] = MeshBoneIndex;
+			OutSkinWeights[VertexIndex].BoneWeights[Index] = BoneWeight;
+			OutSkinWeights[VertexIndex].MeshBoneIndices[Index] = MeshBoneIndex;
 		}
 	}
 }
 
-void GetReducedSkinWeightsData(
-	const USkeletalMesh& SkeletalMesh, const int32 SkeletalLODIndex, 
-	const TArray<int32> StaticToSkelMapping, 
-	TArray<TVertexSkinWeight<4>>& OutSkinWeights)
-{	
-	// Get SkeletalMesh Weights
-	TArray<TVertexSkinWeight<MAX_TOTAL_INFLUENCES>> SkinWeights;
-	AnimToTexture_Private::GetSkinWeightsData(SkeletalMesh, SkeletalLODIndex, SkinWeights);
-	
-	// Reduce SkeletalMesh Weights to 4 highest influences.
-	TArray<TVertexSkinWeight<4>> ReducedSkinWeights;
-	AnimToTexture_Private::ReduceSkinWeightsData(SkinWeights, ReducedSkinWeights);
-
-	// Allocate StaticMesh Weights
-	OutSkinWeights.SetNumUninitialized(StaticToSkelMapping.Num());
+void ReduceSkinWeights(const TArray<VertexSkinWeightMax>& InSkinWeights, 
+	TArray<VertexSkinWeightFour>& OutSkinWeights)
+{
+	OutSkinWeights.SetNumUninitialized(InSkinWeights.Num());
 
-	// Loop thru SkeletalMesh Weights
-	for (int32 VertexIndex = 0; VertexIndex < StaticToSkelMapping.Num(); ++VertexIndex)
+	for (int32 VertexIndex = 0; VertexIndex < InSkinWeights.Num(); VertexIndex++)
 	{
-		const int32 SkelVertexIndex = StaticToSkelMapping[VertexIndex];
-		OutSkinWeights[VertexIndex] = ReducedSkinWeights[SkelVertexIndex];
+		const VertexSkinWeightMax& InVertexSkinWeight = InSkinWeights[VertexIndex];
+		VertexSkinWeightFour& OutVertexSkinWeight = OutSkinWeights[VertexIndex];
+
+		float TotalBoneWeight = 0.f;
+		for (int32 Index = 0; Index < 4; Index++)
+		{
+			TotalBoneWeight += (float)InVertexSkinWeight.BoneWeights[Index] / 255.f;
+		}
+
+		// We assume Weights are sorted
+		for (int32 Index = 0; Index < 4; Index++)
+		{
+			OutVertexSkinWeight.BoneWeights[Index] = (uint8)FMath::RoundToInt((float)InVertexSkinWeight.BoneWeights[Index] / TotalBoneWeight);
+			OutVertexSkinWeight.MeshBoneIndices[Index] = InVertexSkinWeight.MeshBoneIndices[Index];
+		}
 	}
 }
 
-} // end namespace AnimToTexture_Private
\ No newline at end of file
+void GetSkinnedVertices(const USkeletalMeshComponent* SkeletalMeshComponent, const int32 LODIndex,
+	TArray<FVector3f>& OutPositions)
+{
+	check(SkeletalMeshComponent);
+	OutPositions.Reset();
+
+	// Get SkeletalMesh
+	const USkeletalMesh* SkeletalMesh = SkeletalMeshComponent->GetSkeletalMeshAsset();
+	check(SkeletalMesh);
+
+	// Get Matrices
+	TArray<FMatrix44f> RefToLocals;
+	SkeletalMeshComponent->CacheRefToLocalMatrices(RefToLocals);
+
+	// Get Ref-Pose Vertices
+	TArray<FVector3f> Vertices;
+	const int32 NumVertices = GetVertices(SkeletalMesh, LODIndex, Vertices);
+	OutPositions.SetNumUninitialized(NumVertices);
+
+	// TODO: Add Morph Deltas to Vertices.
+
+	// Get Weights
+	TArray<VertexSkinWeightMax> SkinWeights;
+	GetSkinWeights(SkeletalMesh, LODIndex, SkinWeights);
+
+	for (int32 VertexIndex = 0; VertexIndex < NumVertices; VertexIndex++)
+	{
+		const FVector3f& Vertex = Vertices[VertexIndex];
+		const VertexSkinWeightMax& Weights = SkinWeights[VertexIndex];
+
+		FVector4f SkinnedVertex(0);
+		for (int32 Index = 0; Index < MAX_TOTAL_INFLUENCES; Index++)
+		{
+			const uint8& BoneWeight = Weights.BoneWeights[Index];
+			const uint16& MeshBoneIndex = Weights.MeshBoneIndices[Index];
+
+			// Get Matrix
+			const FMatrix44f& RefToLocal = RefToLocals[MeshBoneIndex];
+
+			const float Weight = (float)BoneWeight / 255.f;
+			SkinnedVertex += RefToLocal.TransformPosition(Vertex) * Weight;
+		}
+
+		OutPositions[VertexIndex] = SkinnedVertex;
+	};
+};
+
+} // end namespace AnimToTexture_Private
diff --git a/Engine/Plugins/Animation/AnimToTexture/Source/AnimToTextureEditor/Private/AnimToTextureUtils.cpp b/Engine/Plugins/Animation/AnimToTexture/Source/AnimToTextureEditor/Private/AnimToTextureUtils.cpp
--- a/Engine/Plugins/Animation/AnimToTexture/Source/AnimToTextureEditor/Private/AnimToTextureUtils.cpp
+++ b/Engine/Plugins/Animation/AnimToTexture/Source/AnimToTextureEditor/Private/AnimToTextureUtils.cpp
@@ -1,83 +1,70 @@
 // Copyright Epic Games, Inc. All Rights Reserved.
 
 #include "AnimToTextureUtils.h"
-#include "AnimToTextureSkeletalMesh.h"
-
-#include "Rendering/SkeletalMeshRenderData.h"
 
 namespace AnimToTexture_Private
 {
-	bool CheckSkinWeightsToTexture(const TArray<TVertexSkinWeight<4>>& SkinWeights)
-	{
-		bool bHasInvalidIndices = false;
-		int16 InvalidMeshBoneIndex = INDEX_NONE;
-		for (int32 VertexIndex = 0; VertexIndex < SkinWeights.Num(); ++VertexIndex)
-		{
-			// MeshBoneIndex can not be higher than 255
-			for (uint16 MeshBoneIndex : SkinWeights[VertexIndex].MeshBoneIndices)
-			{
-				if (MeshBoneIndex > 255)
-				{
-					InvalidMeshBoneIndex = MeshBoneIndex;
-					bHasInvalidIndices = true;
-					break;
-				}
-			}
 
-			if (InvalidMeshBoneIndex != INDEX_NONE)
-			{
-				UE_LOG(LogTemp, Warning, TEXT("Invalid SkinWeights on Vertex: %i. MeshBoneIndex: %i > 255"), VertexIndex, InvalidMeshBoneIndex)
-			}
-		}
+void DecomposeTransformation(const FTransform& Transform, 
+	FVector3f& OutTranslation, FVector4& OutRotation)
+{
+	// Get Translation
+	OutTranslation = (FVector3f)Transform.GetTranslation();
+
+	// Get Rotation 
+	const FQuat Quat = Transform.GetRotation();
 
-		if (bHasInvalidIndices)
-		{
-			return false;
-		}
+	FVector Axis;
+	float Angle;
+	Quat.ToAxisAndAngle(Axis, Angle);
 
-		return true;
-	}
+	OutRotation = FVector4(Axis, Angle);
+}
 
-	bool WriteSkinWeightsToTexture(const TArray<TVertexSkinWeight<4>>& SkinWeights,
-		const int32 RowsPerFrame, const int32 Height, const int32 Width, UTexture2D* Texture)
+void DecomposeTransformations(const TArray<FTransform>& Transforms, 
+	TArray<FVector3f>& OutTranslations, TArray<FVector4>& OutRotations)
+{
+	const int32 NumTransforms = Transforms.Num();
+	OutTranslations.SetNumUninitialized(NumTransforms);
+	OutRotations.SetNumUninitialized(NumTransforms);
+
+	for (int32 Index = 0; Index < NumTransforms; ++Index)
 	{
-		if (!Texture)
-		{
-			return false;
-		}
+		DecomposeTransformation(Transforms[Index], OutTranslations[Index], OutRotations[Index]);
+	}
+};
 
-		// Sanity check for MeshBoneIndices > 256
-		if (!CheckSkinWeightsToTexture(SkinWeights))
-		{
-			return false;
-		}
+bool WriteSkinWeightsToTexture(const TArray<VertexSkinWeightFour>& SkinWeights,
+	const int32 RowsPerFrame, const int32 Height, const int32 Width, UTexture2D* Texture)
+{
+	check(Texture);
 
-		const int32 NumVertices = SkinWeights.Num();
+	const int32 NumVertices = SkinWeights.Num();
 
-		// Allocate PixelData.
-		TArray<FLowPrecision::ColorType> Pixels;
-		Pixels.Init(FLowPrecision::ColorType::Black, Height * Width);
+	// Allocate PixelData.
+	TArray<FLowPrecision::ColorType> Pixels;
+	Pixels.Init(FLowPrecision::ColorType::Black, Height * Width);
 
-		for (int32 VertexIndex = 0; VertexIndex < NumVertices; ++VertexIndex)
-		{
-			const TVertexSkinWeight<4>& VertexSkinWeight = SkinWeights[VertexIndex];
+	for (int32 VertexIndex = 0; VertexIndex < NumVertices; ++VertexIndex)
+	{
+		const VertexSkinWeightFour& VertexSkinWeight = SkinWeights[VertexIndex];
 
-			// Write Influence
-			// NOTE: we are assuming the bone index is under < 256
-			Pixels[VertexIndex].R = (uint8)VertexSkinWeight.MeshBoneIndices[0];
-			Pixels[VertexIndex].G = (uint8)VertexSkinWeight.MeshBoneIndices[1];
-			Pixels[VertexIndex].B = (uint8)VertexSkinWeight.MeshBoneIndices[2];
-			Pixels[VertexIndex].A = (uint8)VertexSkinWeight.MeshBoneIndices[3];
+		// Write Influence
+		// NOTE: we are assuming the bone index is under < 256
+		Pixels[VertexIndex].R = (uint8)VertexSkinWeight.MeshBoneIndices[0];
+		Pixels[VertexIndex].G = (uint8)VertexSkinWeight.MeshBoneIndices[1];
+		Pixels[VertexIndex].B = (uint8)VertexSkinWeight.MeshBoneIndices[2];
+		Pixels[VertexIndex].A = (uint8)VertexSkinWeight.MeshBoneIndices[3];
 			
-			// Write Weight
-			Pixels[Width * RowsPerFrame + VertexIndex].R = VertexSkinWeight.BoneWeights[0];
-			Pixels[Width * RowsPerFrame + VertexIndex].G = VertexSkinWeight.BoneWeights[1];
-			Pixels[Width * RowsPerFrame + VertexIndex].B = VertexSkinWeight.BoneWeights[2];
-			Pixels[Width * RowsPerFrame + VertexIndex].A = VertexSkinWeight.BoneWeights[3];
-		};
+		// Write Weight
+		Pixels[Width * RowsPerFrame + VertexIndex].R = VertexSkinWeight.BoneWeights[0];
+		Pixels[Width * RowsPerFrame + VertexIndex].G = VertexSkinWeight.BoneWeights[1];
+		Pixels[Width * RowsPerFrame + VertexIndex].B = VertexSkinWeight.BoneWeights[2];
+		Pixels[Width * RowsPerFrame + VertexIndex].A = VertexSkinWeight.BoneWeights[3];
+	};
 
-		// Write to Texture
-		return WriteToTexture<FLowPrecision>(Texture, Height, Width, Pixels);
-	}
+	// Write to Texture
+	return WriteToTexture<FLowPrecision>(Texture, Height, Width, Pixels);
+}
 
-}
\ No newline at end of file
+} // end namespace AnimToTexture_Private
\ No newline at end of file
diff --git a/Engine/Plugins/Animation/AnimToTexture/Source/AnimToTextureEditor/Public/AnimToTextureBPLibrary.h b/Engine/Plugins/Animation/AnimToTexture/Source/AnimToTextureEditor/Public/AnimToTextureBPLibrary.h
--- a/Engine/Plugins/Animation/AnimToTexture/Source/AnimToTextureEditor/Public/AnimToTextureBPLibrary.h
+++ b/Engine/Plugins/Animation/AnimToTexture/Source/AnimToTextureEditor/Public/AnimToTextureBPLibrary.h
@@ -3,6 +3,7 @@
 #pragma once
 
 #include "AnimToTextureDataAsset.h"
+#include "AnimToTextureMeshMapping.h"
 #include "Kismet/BlueprintFunctionLibrary.h"
 #include "Materials/MaterialLayersFunctions.h"
 
@@ -28,7 +29,7 @@ class UAnimToTextureBPLibrary : public UBlueprintFunctionLibrary
 	* Bakes Animation Data into Textures.
 	*/
 	UFUNCTION(BlueprintCallable, meta = (Category = "AnimToTexture"))
-	static void AnimationToTexture(UAnimToTextureDataAsset* DataAsset, const FTransform RootTransform);
+	static void AnimationToTexture(UAnimToTextureDataAsset* DataAsset, const FTransform RootTransform, bool& bSuccess);
 
 	/** 
 	* Utility for converting SkeletalMesh into a StaticMesh
@@ -36,6 +37,10 @@ class UAnimToTextureBPLibrary : public UBlueprintFunctionLibrary
 	UFUNCTION(BlueprintCallable, Category = "AnimToTexture")
 	static UStaticMesh* ConvertSkeletalMeshToStaticMesh(USkeletalMesh* SkeletalMesh, const FString PackageName, const int32 LODIndex = -1);
 
+	// FIXME: you cannot set index to 2 if there is no index 1
+	UFUNCTION(BlueprintCallable, Category = "AnimToTexture")
+	static bool SetLightMapIndex(UStaticMesh* StaticMesh, const int32 LODIndex, const int32 LightmapIndex=1, bool bGenerateLightmapUVs=true);
+
 	/**
 	* Updates a material's parameters to match those of an animToTexture data asset
 	*/
@@ -48,19 +53,36 @@ class UAnimToTextureBPLibrary : public UBlueprintFunctionLibrary
 
 private:
 
+	// Get Vertex and Normals from Current Pose
+	// The VertexDelta is returned from the RefPose
+	static void GetVertexDeltasAndNormals(const USkeletalMeshComponent* SkeletalMeshComponent, const int32 LODIndex, 
+		const TArray<FVector3f>& Vertices, const TArray<FVector3f>& Normals, const TArray<AnimToTexture_Private::FVertexToMeshMapping>& Mapping,
+		const FTransform RootTransform,
+		TArray<FVector3f>& OutVertexDeltas, TArray<FVector3f>& OutVertexNormals);
+
+	// Gets RefPose Bone Position and Rotations.
+	static int32 GetRefBonePositionsAndRotations(const USkeletalMeshComponent* SkeletalMeshComponent, 
+		TArray<FVector3f>& OutBoneRefPositions, TArray<FVector4>& OutBoneRefRotations);
+
+	// Gets Bone Position and Rotations for Current Pose.	
+	// The BonePosition is returned relative to the RefPose
+	static int32 GetBonePositionsAndRotations(const USkeletalMeshComponent* SkeletalMeshComponent, const TArray<FVector3f>& BoneRefPositions,
+		TArray<FVector3f>& BonePositions, TArray<FVector4>& BoneRotations);
+
+	// Normalizes Deltas and Normals between [0-1] with Bounding Box
 	static void NormalizeVertexData(
 		const TArray<FVector3f>& Deltas, const TArray<FVector3f>& Normals,
-		FVector& MinBBox, FVector& SizeBBox,
-		TArray<FVector3f>& NormalizedDeltas, TArray<FVector3f>& NormalizedNormals);
+		FVector& OutMinBBox, FVector& OutSizeBBox,
+		TArray<FVector3f>& OutNormalizedDeltas, TArray<FVector3f>& OutNormalizedNormals);
 
+	// Normalizes Positions and Rotations between [0-1] with Bounding Box
 	static void NormalizeBoneData(
 		const TArray<FVector3f>& Positions, const TArray<FVector4>& Rotations,
-		FVector& MinBBox, FVector& SizeBBox,
-		TArray<FVector3f>& NormalizedPositions, TArray<FVector4>& NormalizedRotations);
+		FVector& OutMinBBox, FVector& OutSizeBBox,
+		TArray<FVector3f>& OutNormalizedPositions, TArray<FVector4>& OutNormalizedRotations);
 
 	/* Returns best resolution for the given data. 
-	*  Returns false if data wont fit in the the max range.
-	*/
+	*  Returns false if data doesnt fit in the the max range */
 	static bool FindBestResolution(const int32 NumFrames, const int32 NumElements,
 								   int32& OutHeight, int32& OutWidth, int32& OutRowsPerFrame,
 								   const int32 MaxHeight = 4096, const int32 MaxWidth = 4096, bool bEnforcePowerOfTwo = false);
diff --git a/Engine/Plugins/Animation/AnimToTexture/Source/AnimToTextureEditor/Public/AnimToTextureMeshMapping.h b/Engine/Plugins/Animation/AnimToTexture/Source/AnimToTextureEditor/Public/AnimToTextureMeshMapping.h
new file mode 100644
--- /dev/null
+++ b/Engine/Plugins/Animation/AnimToTexture/Source/AnimToTextureEditor/Public/AnimToTextureMeshMapping.h
@@ -0,0 +1,57 @@
+// Copyright Epic Games, Inc. All Rights Reserved.
+
+#pragma once
+
+#include "AnimToTextureSkeletalMesh.h"
+#include "CoreMinimal.h"
+
+namespace AnimToTexture_Private
+{
+
+// Static to Skeletal Mapping
+struct FVertexToMeshMapping
+{
+	// Closest SkeletalMesh Triangle
+	FIntVector3 Triangle;
+	// SkeletalMesh Triangle Barycentric Coordinates
+	FVector3f  BarycentricCoords;
+	// SkeletalMesh Triangle Inverted Matrix
+	FMatrix44f InvMatrix;
+	
+	/* Transforms Point with given Triangle */
+	FVector3f TransformPosition(const FVector3f& Point, const FVector3f& PointA, const FVector3f& PointB, const FVector3f& PointC) const;
+	
+	/* Transforms Vector with given Triangle */
+	FVector3f TransformVector(const FVector3f& Vector, const FVector3f& PointA, const FVector3f& PointB, const FVector3f& PointC) const;
+	
+	/* Interpolate VertexSkinWeights with Barycentric Coordinates */
+	void InterpolateVertexSkinWeights(const VertexSkinWeightMax& A, const VertexSkinWeightMax& B, const VertexSkinWeightMax& C,
+		VertexSkinWeightMax& OutWeights) const;
+	
+	/* Computes closest point to triangle*/
+	static FVector3f FindClosestPointToTriangle(const FVector3f& Point, const FVector3f& PointA, const FVector3f& PointB, const FVector3f& PointC);
+
+	/* Computes Barycentric coordinates from point to triangle */
+	static FVector3f BarycentricCoordinates(const FVector3f& Point, const FVector3f& PointA, const FVector3f& PointB, const FVector3f& PointC);
+
+	/* Computes Triangle Normal */
+	static FVector3f GetTriangleNormal(const FVector3f& PointA, const FVector3f& PointB, const FVector3f& PointC);
+
+	/* Computes Triangle Matrix */
+	static FMatrix44f GetTriangleMatrix(const FVector3f& PointA, const FVector3f& PointB, const FVector3f& PointC);
+
+	/* Finds closest triangle, point on triangle and barycentric coordinates */
+	static int32 FindClosestTriangle(const FVector3f& Point, const TArray<FVector3f>& Vertices, const TArray<FIntVector3>& Triangles,
+		FVector3f& OutClosestPoint, FVector3f& OutBarycentricCoords);
+
+	// Creates Mapping between StaticMesh and SkeletalMesh
+	static void Create(const UStaticMesh* StaticMesh, const int32 StaticMeshLODIndex,
+		const USkeletalMesh* SkeletalMesh, const int32 SkeletalMeshLODIndex,
+		TArray<FVertexToMeshMapping>& OutMapping);
+
+	/* Interpolate SkinWeights with VertexToMeshMapping */
+	static void InterpolateSkinWeights(const TArray<FVertexToMeshMapping>& Mapping, const TArray<VertexSkinWeightMax>& SkinWeights,
+		TArray<VertexSkinWeightMax>& OutSkinWeights);
+};
+
+} // end namespace AnimToTexture_Private
diff --git a/Engine/Plugins/Animation/AnimToTexture/Source/AnimToTextureEditor/Public/AnimToTextureSkeletalMesh.h b/Engine/Plugins/Animation/AnimToTexture/Source/AnimToTextureEditor/Public/AnimToTextureSkeletalMesh.h
--- a/Engine/Plugins/Animation/AnimToTexture/Source/AnimToTextureEditor/Public/AnimToTextureSkeletalMesh.h
+++ b/Engine/Plugins/Animation/AnimToTexture/Source/AnimToTextureEditor/Public/AnimToTextureSkeletalMesh.h
@@ -2,11 +2,10 @@
 
 #pragma once
 
+#include "AnimToTextureDataAsset.h"
 #include "CoreMinimal.h"
-#include "GPUSkinPublicDefs.h"
 #include "Containers/StaticArray.h"
-
-#include "AnimToTextureDataAsset.h"
+#include "GPUSkinPublicDefs.h" 
 
 namespace AnimToTexture_Private
 {
@@ -15,44 +14,34 @@ template <uint16 NumInfluences>
 struct TVertexSkinWeight
 {
 	TStaticArray<uint16, NumInfluences> MeshBoneIndices;
-	TStaticArray<uint16, NumInfluences> BoneIndices;
 	TStaticArray<uint8, NumInfluences>  BoneWeights;
 };
 
-/* Returns RefPose Vertex Positions */
-int32 GetVertices(const USkeletalMesh& Mesh, const int32 LODIndex, 
+using VertexSkinWeightMax  = TVertexSkinWeight<MAX_TOTAL_INFLUENCES>;
+using VertexSkinWeightFour = TVertexSkinWeight<4>;
+
+/* Returns StaticMesh Vertex Positions */
+int32 GetVertices(const UStaticMesh* StaticMesh, const int32 LODIndex,
 	TArray<FVector3f>& OutPositions, TArray<FVector3f>& OutNormals);
 
-/* Finds closest vertex to point */
-int32 FindClosestVertex(const FVector3f& Point, const TArray<FVector3f>& Vertices);
+/* Returns RefPose Vertex Positions */
+int32 GetVertices(const USkeletalMesh* SkeletalMesh, const int32 LODIndex,
+	TArray<FVector3f>& OutPositions);
 
-/* VertexMapping between StaticMesh and SkeletalMesh Vertices.
-*  Note: Size of Mapping is Number of MeshDescription Vertices.
-*/
-void GetStaticToSkeletalMapping(
-	const UStaticMesh& StaticMesh, const int32 StaticLODIndex,
-	const USkeletalMesh& SkeletalMesh, const int32 SkeletalLODIndex,
-	TArray<FVector3f>& OutStaticVertices, TArray<FVector3f>& OutSkeletalVertices,
-	TArray<int32>& OutStaticToSkeletalMapping);
+/* Returns Triangles vertex indices */
+int32 GetTriangles(const USkeletalMesh* SkeletalMesh, const int32 LODIndex,
+	TArray<FIntVector3>& OutTriangles);
 
 /* Computes CPUSkinning at Pose */
-void GetSkinnedVertices(USkeletalMeshComponent* MeshComponent, const int32 LODIndex,
-	TArray<FVector3f>& OutPositions, TArray<FVector3f>& OutNormals);
+void GetSkinnedVertices(const USkeletalMeshComponent* SkeletalMeshComponent, const int32 LODIndex,
+	TArray<FVector3f>& OutPositions);
 
 /** Gets Skin Weights Data from SkeletalMeshComponent */
-void GetSkinWeightsData(const USkeletalMesh& SkeletalMesh, const int32 LODIndex, TArray<TVertexSkinWeight<MAX_TOTAL_INFLUENCES>>& SkinWeights);
+void GetSkinWeights(const USkeletalMesh* SkeletalMesh, const int32 LODIndex, 
+	TArray<VertexSkinWeightMax>& OutSkinWeights);
 
-/** Gets Skin Weights Data from SkeletalMeshComponent, Reduce it to NumInfluences and Maps it to the StaticMesh */
-void GetReducedSkinWeightsData(
-	const USkeletalMesh& SkeletalMesh, const int32 SkeletalLODIndex, 
-	const TArray<int32> StaticToSkelMapping,
-	TArray<TVertexSkinWeight<4>>& SkinWeights);
-
-/** Reduce Weights */
-void ReduceSkinWeightsData(const TArray<TVertexSkinWeight<MAX_TOTAL_INFLUENCES>>& InSkinWeights, TArray<TVertexSkinWeight<4>>& OutSkinWeights);
-
-/* Returns the 4-Higher Influences and Weights */
-void GetMaxFourInfluences(const TVertexSkinWeight<MAX_TOTAL_INFLUENCES>& InSkinWeights, TVertexSkinWeight<4>& OutSkinWeights);
+/** Reduce Weights from MAX_TOTAL_INFLUENCES to 4 */
+void ReduceSkinWeights(const TArray<VertexSkinWeightMax>& InSkinWeights, TArray<VertexSkinWeightFour>& OutSkinWeights);
 
 /* Returns Number of RawBones (no virtual bones)*/
 int32 GetNumBones(const USkeletalMesh* SkeletalMesh);
@@ -69,82 +58,4 @@ bool HasBone(const USkeletalMesh* SkeletalMesh, const FName& Bone);
    Only the RawBones are returned (no virtual bones) */
 void GetBoneNames(const USkeletalMesh* SkeletalMesh, TArray<FName>& OutNames);
 
-/* Decomposes Transform in Translation and AxisAndAngle */
-void DecomposeTransformation(const FTransform& Transform, FVector3f& Translation, FVector4& Rotation);
-void DecomposeTransformations(const TArray<FTransform>& Transforms, TArray<FVector3f>& Translations, TArray<FVector4>& Rotations);
-
-// ----------------------------------------------------------------------------
-
-FORCEINLINE void GetMaxFourInfluences(
-	const TVertexSkinWeight<MAX_TOTAL_INFLUENCES>& InSkinWeights, TVertexSkinWeight<4>& OutSkinWeights)
-{
-	// Reset Values
-	for (int32 Index = 0; Index < 4; ++Index)
-	{
-		OutSkinWeights.BoneWeights[Index] = 0;
-		OutSkinWeights.BoneIndices[Index] = 0;
-		OutSkinWeights.MeshBoneIndices[Index] = 0;
-	}
-
-	// Find 1st Influence (Highest)
-	for (int32 Index = 0; Index < MAX_TOTAL_INFLUENCES; ++Index)
-	{
-		if (InSkinWeights.BoneWeights[Index] > OutSkinWeights.BoneWeights[0])
-		{
-			OutSkinWeights.BoneWeights[0] = InSkinWeights.BoneWeights[Index];
-			OutSkinWeights.BoneIndices[0] = InSkinWeights.BoneIndices[Index];
-			OutSkinWeights.MeshBoneIndices[0] = InSkinWeights.MeshBoneIndices[Index];
-		}
-	}
-
-	// Find 2nd Influence
-	for (int32 Index = 0; Index < MAX_TOTAL_INFLUENCES; ++Index)
-	{		
-		if (InSkinWeights.BoneIndices[Index] != OutSkinWeights.BoneIndices[0])
-		{
-			if (InSkinWeights.BoneWeights[Index] > OutSkinWeights.BoneWeights[1])
-			{
-				OutSkinWeights.BoneWeights[1] = InSkinWeights.BoneWeights[Index];
-				OutSkinWeights.BoneIndices[1] = InSkinWeights.BoneIndices[Index];
-				OutSkinWeights.MeshBoneIndices[1] = InSkinWeights.MeshBoneIndices[Index];
-			}
-		}
-	}
-
-	// Find 3rd Influence
-	for (int32 Index = 0; Index < MAX_TOTAL_INFLUENCES; ++Index)
-	{
-		if (InSkinWeights.BoneIndices[Index] != OutSkinWeights.BoneIndices[0] && InSkinWeights.BoneIndices[Index] != OutSkinWeights.BoneIndices[1])
-		{
-			if (InSkinWeights.BoneWeights[Index] > OutSkinWeights.BoneWeights[2])
-			{
-				OutSkinWeights.BoneWeights[2] = InSkinWeights.BoneWeights[Index];
-				OutSkinWeights.BoneIndices[2] = InSkinWeights.BoneIndices[Index];
-				OutSkinWeights.MeshBoneIndices[2] = InSkinWeights.MeshBoneIndices[Index];
-			}
-		}
-	}
-
-	// Find 4th Influence
-	for (int32 Index = 0; Index < MAX_TOTAL_INFLUENCES; ++Index)
-	{
-		if (InSkinWeights.BoneIndices[Index] != OutSkinWeights.BoneIndices[0] && InSkinWeights.BoneIndices[Index] != OutSkinWeights.BoneIndices[1] && InSkinWeights.BoneIndices[Index] != OutSkinWeights.BoneIndices[2])
-		{
-			if (InSkinWeights.BoneWeights[Index] > OutSkinWeights.BoneWeights[3])
-			{
-				OutSkinWeights.BoneWeights[3] = InSkinWeights.BoneWeights[Index];
-				OutSkinWeights.BoneIndices[3] = InSkinWeights.BoneIndices[Index];
-				OutSkinWeights.MeshBoneIndices[3] = InSkinWeights.MeshBoneIndices[Index];
-			}
-		}
-	}
-
-	// Normalize Weighs
-	const float NormFactor = 255.0f / (OutSkinWeights.BoneWeights[0] + OutSkinWeights.BoneWeights[1] + OutSkinWeights.BoneWeights[2] + OutSkinWeights.BoneWeights[3]);
-	for (int32 Index = 0; Index < 4; ++Index)
-	{
-		OutSkinWeights.BoneWeights[Index] = FMath::RoundToInt(OutSkinWeights.BoneWeights[Index] * NormFactor);
-	}
-}
-
 } // end namespace AnimToTexture_Private
diff --git a/Engine/Plugins/Animation/AnimToTexture/Source/AnimToTextureEditor/Public/AnimToTextureUtils.h b/Engine/Plugins/Animation/AnimToTexture/Source/AnimToTextureEditor/Public/AnimToTextureUtils.h
--- a/Engine/Plugins/Animation/AnimToTexture/Source/AnimToTextureEditor/Public/AnimToTextureUtils.h
+++ b/Engine/Plugins/Animation/AnimToTexture/Source/AnimToTextureEditor/Public/AnimToTextureUtils.h
@@ -2,80 +2,83 @@
 
 #pragma once
 
+#include "AnimToTextureSkeletalMesh.h"
 #include "CoreMinimal.h"
+#include "TextureResource.h"
 #include "Engine/Texture.h"
 #include "Engine/Texture2D.h"
-#include "TextureResource.h"
+
 
 namespace AnimToTexture_Private
 {
-	// Forward declaration
-	template <uint16 NumInfluences> struct TVertexSkinWeight;
 
-	struct FVector4u16
-	{
-		uint16 X;
-		uint16 Y;
-		uint16 Z;
-		uint16 W;
-	};
+struct FVector4u16
+{
+	uint16 X;
+	uint16 Y;
+	uint16 Z;
+	uint16 W;
+};
 
-	struct FLowPrecision
-	{
-		using ColorType = FColor;
-		static constexpr EPixelFormat PixelFormat = EPixelFormat::PF_B8G8R8A8;
-		static constexpr ETextureSourceFormat TextureSourceFormat = ETextureSourceFormat::TSF_BGRA8;
-		static constexpr TextureCompressionSettings CompressionSettings = TextureCompressionSettings::TC_VectorDisplacementmap;
-		static constexpr ColorType DefaultColor = FColor(0, 0, 0);		
-	};
-
-	struct FHighPrecision
-	{
-		using ColorType = FVector4u16;
-		static constexpr EPixelFormat PixelFormat = EPixelFormat::PF_R16G16B16A16_UNORM;
-		static constexpr ETextureSourceFormat TextureSourceFormat = ETextureSourceFormat::TSF_RGBA16;
-		static constexpr TextureCompressionSettings CompressionSettings = TextureCompressionSettings::TC_HDR;
-		static constexpr ColorType DefaultColor = { 0, 0, 0, 0 };
-	};
-
-
-	/** Writes list of vectors into texture
-	*   Note: They must be pre-normalized. */
-	template<class V, class TextureSettings>
-	bool WriteVectorsToTexture(const TArray<V>& Vectors,
-		const int32 NumFrames, const int32 RowsPerFrame,
-		const int32 Height, const int32 Width, 
-		UTexture2D* Texture);
-
-	/* Writes list of skinweights into texture.
-	*  The SkinWeights data is already in uint8 & uint16 format, no need for normalizing it.
-	*/
-	bool WriteSkinWeightsToTexture(const TArray<TVertexSkinWeight<4>>& SkinWeights,
-		const int32 RowsPerFrame,
-		const int32 Height, const int32 Width,
-		UTexture2D* Texture);
-
-	/* Helper utility for writing 8 or 16 bits textures */
-	template<class TextureSettings>
-	bool WriteToTexture(UTexture2D* Texture, const uint32 Height, const uint32 Width, const TArray<typename TextureSettings::ColorType>& Data);
-
-	template<class V /* FVector / FVector4 / FIntVector */, class C /* FColor / FLinearColor */>
-	void VectorToColor(const V& Vector, C& Color);
-	
-}
+struct FLowPrecision
+{
+	using ColorType = FColor;
+	static constexpr EPixelFormat PixelFormat = EPixelFormat::PF_B8G8R8A8;
+	static constexpr ETextureSourceFormat TextureSourceFormat = ETextureSourceFormat::TSF_BGRA8;
+	static constexpr TextureCompressionSettings CompressionSettings = TextureCompressionSettings::TC_VectorDisplacementmap;
+	static constexpr ColorType DefaultColor = FColor(0, 0, 0);		
+};
+
+struct FHighPrecision
+{
+	using ColorType = FVector4u16;
+	static constexpr EPixelFormat PixelFormat = EPixelFormat::PF_R16G16B16A16_UNORM;
+	static constexpr ETextureSourceFormat TextureSourceFormat = ETextureSourceFormat::TSF_RGBA16;
+	static constexpr TextureCompressionSettings CompressionSettings = TextureCompressionSettings::TC_HDR;
+	static constexpr ColorType DefaultColor = { 0, 0, 0, 0 };
+};
+
+/** Writes list of vectors into texture
+*   Note: They must be pre-normalized. */
+template<class V, class TextureSettings>
+bool WriteVectorsToTexture(const TArray<V>& Vectors,
+	const int32 NumFrames, const int32 RowsPerFrame,
+	const int32 Height, const int32 Width, 
+	UTexture2D* Texture);
+
+/* Writes list of skinweights into texture.
+*  The SkinWeights data is already in uint8 & uint16 format, no need for normalizing it.
+*/
+bool WriteSkinWeightsToTexture(const TArray<VertexSkinWeightFour>& SkinWeights,
+	const int32 RowsPerFrame,
+	const int32 Height, const int32 Width,
+	UTexture2D* Texture);
+
+/* Helper utility for writing 8 or 16 bits textures */
+template<class TextureSettings>
+bool WriteToTexture(UTexture2D* Texture, const uint32 Height, const uint32 Width, const TArray<typename TextureSettings::ColorType>& Data);
+
+template<class V /* FVector / FVector4 / FIntVector */, class C /* FColor / FLinearColor */>
+void VectorToColor(const V& Vector, C& Color);
+
+/* Decomposes Transform in Translation and AxisAndAngle */
+void DecomposeTransformation(const FTransform& Transform, FVector3f& OutTranslation, FVector4& OutRotation);
+void DecomposeTransformations(const TArray<FTransform>& Transforms, TArray<FVector3f>& OutTranslations, TArray<FVector4>& OutRotations);
+
+} // end namespace AnimToTexture_Private
 
 // ----------------------------------------------------------------------------
 
 template<>
 FORCEINLINE void AnimToTexture_Private::VectorToColor(const FVector3f& Vector, FColor& Color)
 {
-	const float ClampedX = FMath::Clamp(Vector.X, 0.0f, 1.0f);
-	const float ClampedY = FMath::Clamp(Vector.Y, 0.0f, 1.0f);
-	const float ClampedZ = FMath::Clamp(Vector.Z, 0.0f, 1.0f);
+	const float ClampedX = FMath::Clamp(Vector.X, 0.f, 1.f);
+	const float ClampedY = FMath::Clamp(Vector.Y, 0.f, 1.f);
+	const float ClampedZ = FMath::Clamp(Vector.Z, 0.f, 1.f);
 
-	Color.R = (uint8)FMath::RoundToInt(ClampedX * 255.0f);
-	Color.G = (uint8)FMath::RoundToInt(ClampedY * 255.0f);
-	Color.B = (uint8)FMath::RoundToInt(ClampedZ * 255.0f);
+	Color.R = (uint8)FMath::RoundToInt(ClampedX * 255.f);
+	Color.G = (uint8)FMath::RoundToInt(ClampedY * 255.f);
+	Color.B = (uint8)FMath::RoundToInt(ClampedZ * 255.f);
 	Color.A = 255;
 }
 
@@ -89,15 +92,15 @@ FORCEINLINE void AnimToTexture_Private::VectorToColor(const FVector3d& Vector, F
 template<>
 FORCEINLINE void AnimToTexture_Private::VectorToColor(const FVector4& Vector, FColor& Color)
 {
-	const float ClampedX = FMath::Clamp(Vector.X, 0.0f, 1.0f);
-	const float ClampedY = FMath::Clamp(Vector.Y, 0.0f, 1.0f);
-	const float ClampedZ = FMath::Clamp(Vector.Z, 0.0f, 1.0f);
-	const float ClampedW = FMath::Clamp(Vector.W, 0.0f, 1.0f);
+	const float ClampedX = FMath::Clamp(Vector.X, 0.f, 1.f);
+	const float ClampedY = FMath::Clamp(Vector.Y, 0.f, 1.f);
+	const float ClampedZ = FMath::Clamp(Vector.Z, 0.f, 1.f);
+	const float ClampedW = FMath::Clamp(Vector.W, 0.f, 1.f);
 	
-	Color.R = (uint8)FMath::RoundToInt(ClampedX * 255.0f);
-	Color.G = (uint8)FMath::RoundToInt(ClampedY * 255.0f);
-	Color.B = (uint8)FMath::RoundToInt(ClampedZ * 255.0f);
-	Color.A = (uint8)FMath::RoundToInt(ClampedW * 255.0f);
+	Color.R = (uint8)FMath::RoundToInt(ClampedX * 255.f);
+	Color.G = (uint8)FMath::RoundToInt(ClampedY * 255.f);
+	Color.B = (uint8)FMath::RoundToInt(ClampedZ * 255.f);
+	Color.A = (uint8)FMath::RoundToInt(ClampedW * 255.f);
 }
 
 template<>
@@ -133,7 +136,7 @@ FORCEINLINE void AnimToTexture_Private::VectorToColor(const FVector3f& Vector, F
 	Color.R = Vector.X;
 	Color.G = Vector.Y;
 	Color.B = Vector.Z;
-	Color.A = 1.0f;
+	Color.A = 1.f;
 }
 
 template<>
@@ -148,19 +151,19 @@ FORCEINLINE void AnimToTexture_Private::VectorToColor(const FVector4& Vector, FL
 template<>
 FORCEINLINE void AnimToTexture_Private::VectorToColor(const FVector3f& Vector, FVector4u16& Color)
 {
-	Color.X = FMath::RoundToInt(FMath::Clamp(Vector.X, 0.0f, 1.0f) * MAX_uint16);
-	Color.Y = FMath::RoundToInt(FMath::Clamp(Vector.Y, 0.0f, 1.0f) * MAX_uint16);
-	Color.Z = FMath::RoundToInt(FMath::Clamp(Vector.Z, 0.0f, 1.0f) * MAX_uint16);
+	Color.X = FMath::RoundToInt(FMath::Clamp(Vector.X, 0.f, 1.f) * MAX_uint16);
+	Color.Y = FMath::RoundToInt(FMath::Clamp(Vector.Y, 0.f, 1.f) * MAX_uint16);
+	Color.Z = FMath::RoundToInt(FMath::Clamp(Vector.Z, 0.f, 1.f) * MAX_uint16);
 	Color.W = MAX_uint16;
 }
 
 template<>
 FORCEINLINE void AnimToTexture_Private::VectorToColor(const FVector4& Vector, FVector4u16& Color)
 {
-	Color.X = FMath::RoundToInt(FMath::Clamp(Vector.X, 0.0f, 1.0f) * MAX_uint16);
-	Color.Y = FMath::RoundToInt(FMath::Clamp(Vector.Y, 0.0f, 1.0f) * MAX_uint16);
-	Color.Z = FMath::RoundToInt(FMath::Clamp(Vector.Z, 0.0f, 1.0f) * MAX_uint16);
-	Color.W = FMath::RoundToInt(FMath::Clamp(Vector.W, 0.0f, 1.0f) * MAX_uint16);
+	Color.X = FMath::RoundToInt(FMath::Clamp(Vector.X, 0.f, 1.f) * MAX_uint16);
+	Color.Y = FMath::RoundToInt(FMath::Clamp(Vector.Y, 0.f, 1.f) * MAX_uint16);
+	Color.Z = FMath::RoundToInt(FMath::Clamp(Vector.Z, 0.f, 1.f) * MAX_uint16);
+	Color.W = FMath::RoundToInt(FMath::Clamp(Vector.W, 0.f, 1.f) * MAX_uint16);
 }
 
 
@@ -182,12 +185,12 @@ FORCEINLINE bool AnimToTexture_Private::WriteVectorsToTexture(const TArray<V>& V
 	Pixels.Init(TextureSettings::DefaultColor, Height * Width);
 
 	// Fillout Frame Data
-	for (int32 Frame = 0; Frame < NumFrames; ++Frame)
+	for (int32 Frame = 0; Frame < NumFrames; Frame++)
 	{
 		const int32 BlockStart = RowsPerFrame * Width * Frame;
 
 		// Set Data.
-		for (int32 Index = 0; Index < NumElements; ++Index)
+		for (int32 Index = 0; Index < NumElements; Index++)
 		{
 			const V& Vector = Vectors[NumElements * Frame + Index];
 			typename TextureSettings::ColorType& Pixel = Pixels[BlockStart + Index];
@@ -203,8 +206,7 @@ FORCEINLINE bool AnimToTexture_Private::WriteVectorsToTexture(const TArray<V>& V
 template<class TextureSettings>
 FORCEINLINE bool AnimToTexture_Private::WriteToTexture(
 	UTexture2D* Texture,
-	const uint32 Height,
-	const uint32 Width,
+	const uint32 Height, const uint32 Width,
 	const TArray<typename TextureSettings::ColorType>& Pixels)
 {
 	if (!Texture)
