Fix another CIS issue.

diff --git a/Engine/Plugins/Animation/AnimToTexture/Source/AnimToTexture/Public/AnimToTextureDataAsset.h b/Engine/Plugins/Animation/AnimToTexture/Source/AnimToTexture/Public/AnimToTextureDataAsset.h
--- a/Engine/Plugins/Animation/AnimToTexture/Source/AnimToTexture/Public/AnimToTextureDataAsset.h
+++ b/Engine/Plugins/Animation/AnimToTexture/Source/AnimToTexture/Public/AnimToTextureDataAsset.h
@@ -4,7 +4,9 @@
 
 #include "Engine/DataAsset.h"
 #include "Animation/AnimInstance.h"
+#include "Engine/SkeletalMesh.h"
 #include "Engine/StaticMesh.h"
+#include "Engine/Texture2D.h"
 #include "AnimToTextureDataAsset.generated.h"
 
 class USkeletalMesh;
