Including USkeleton header

diff --git a/Engine/Plugins/Animation/AnimToTexture/Source/AnimToTextureEditor/Private/AnimToTextureBPLibrary.cpp b/Engine/Plugins/Animation/AnimToTexture/Source/AnimToTextureEditor/Private/AnimToTextureBPLibrary.cpp
--- a/Engine/Plugins/Animation/AnimToTexture/Source/AnimToTextureEditor/Private/AnimToTextureBPLibrary.cpp
+++ b/Engine/Plugins/Animation/AnimToTexture/Source/AnimToTextureEditor/Private/AnimToTextureBPLibrary.cpp
@@ -11,6 +11,7 @@
 #include "MeshUtilities.h"
 #include "Engine/SkeletalMesh.h"
 #include "Components/SkeletalMeshComponent.h"
+#include "Animation/Skeleton.h"
 #include "Animation/AnimSequence.h"
 #include "Math/Vector.h"
 #include "Math/NumericLimits.h"
