AnimToTexture:
- Refactored PlaybackHelpers
- Added prefix to Structs
- Deprecated CitySample specific functionality
- Fixed SetLightMapIndex
- Added Log

diff --git a/Engine/Plugins/Animation/AnimToTexture/Config/BaseAnimToTexture.ini b/Engine/Plugins/Animation/AnimToTexture/Config/BaseAnimToTexture.ini
new file mode 100644
--- /dev/null
+++ b/Engine/Plugins/Animation/AnimToTexture/Config/BaseAnimToTexture.ini
@@ -0,0 +1,3 @@
+[CoreRedirects]
++StructRedirects=(OldName="/Script/AnimToTexture.AnimSequenceInfo",NewName="/Script/AnimToTexture.AnimToTextureAnimSequenceInfo")
++StructRedirects=(OldName="/Script/AnimToTexture.AnimInfo",NewName="/Script/AnimToTexture.AnimToTextureAnimInfo")
diff --git a/Engine/Plugins/Animation/AnimToTexture/Source/AnimToTexture/Private/AnimToTexture.cpp b/Engine/Plugins/Animation/AnimToTexture/Source/AnimToTexture/Private/AnimToTexture.cpp
deleted file mode 100644
diff --git a/Engine/Plugins/Animation/AnimToTexture/Source/AnimToTexture/Private/AnimToTextureDataAsset.cpp b/Engine/Plugins/Animation/AnimToTexture/Source/AnimToTexture/Private/AnimToTextureDataAsset.cpp
--- a/Engine/Plugins/Animation/AnimToTexture/Source/AnimToTexture/Private/AnimToTextureDataAsset.cpp
+++ b/Engine/Plugins/Animation/AnimToTexture/Source/AnimToTexture/Private/AnimToTextureDataAsset.cpp
@@ -4,22 +4,19 @@
 
 int32 UAnimToTextureDataAsset::GetIndexFromAnimSequence(const UAnimSequence* Sequence)
 {
-	int32 OutIndex = 0;
-
-	int32 NumSequences = AnimSequences.Num();
+	const int32 NumSequences = AnimSequences.Num();
 	
 	// We can store a sequence to index map for a faster search
-	for (int32 CurrentIndex = 0; CurrentIndex < NumSequences; ++CurrentIndex)
+	for (int32 Index = 0; Index < NumSequences; ++Index)
 	{
-		const FAnimSequenceInfo& SequenceInfo = AnimSequences[CurrentIndex];
-		if (SequenceInfo.AnimSequence == Sequence)
+		const FAnimToTextureAnimSequenceInfo& SequenceInfo = AnimSequences[Index];
+		if (SequenceInfo.bEnabled && SequenceInfo.AnimSequence == Sequence)
 		{
-			OutIndex = CurrentIndex;
-			break;
+			return Index;
 		}
 	}
 
-	return OutIndex;
+	return INDEX_NONE;
 }
 
 void UAnimToTextureDataAsset::ResetInfo()
diff --git a/Engine/Plugins/Animation/AnimToTexture/Source/AnimToTexture/Private/AnimToTextureInstancePlaybackHelpers.cpp b/Engine/Plugins/Animation/AnimToTexture/Source/AnimToTexture/Private/AnimToTextureInstancePlaybackHelpers.cpp
--- a/Engine/Plugins/Animation/AnimToTexture/Source/AnimToTexture/Private/AnimToTextureInstancePlaybackHelpers.cpp
+++ b/Engine/Plugins/Animation/AnimToTexture/Source/AnimToTexture/Private/AnimToTextureInstancePlaybackHelpers.cpp
@@ -3,81 +3,110 @@
 #include "AnimToTextureInstancePlaybackHelpers.h"
 #include "AnimToTextureDataAsset.h"
 
-void UAnimToTextureInstancePlaybackLibrary::SetupInstancedMeshComponent(UInstancedStaticMeshComponent* InstancedMeshComponent, FAnimToTextureInstanceData& InstanceData, int32 NumInstances)
-{
-	if (InstancedMeshComponent)
+bool UAnimToTextureInstancePlaybackLibrary::SetupInstancedMeshComponent(UInstancedStaticMeshComponent* InstancedMeshComponent, int32 NumInstances, bool bAutoPlay)
+{	
+	if (!InstancedMeshComponent)
 	{
-		InstancedMeshComponent->NumCustomDataFloats = InstanceData.PlaybackData.GetTypeSize() / sizeof(float);
-		InstancedMeshComponent->PerInstanceSMData.Reset();
-		InstancedMeshComponent->PerInstanceSMData.AddDefaulted(NumInstances);
-		InstancedMeshComponent->PreAllocateInstancesMemory(NumInstances);
-		InstancedMeshComponent->PerInstanceSMCustomData.SetNumZeroed(NumInstances * InstancedMeshComponent->NumCustomDataFloats);
-		AllocateInstanceData(InstanceData, NumInstances);
+		return false;
 	}
-}
 
-void UAnimToTextureInstancePlaybackLibrary::BatchUpdateInstancedMeshComponent(UInstancedStaticMeshComponent* InstancedMeshComponent, FAnimToTextureInstanceData& InstanceData)
-{
-	SIZE_T CustomDataSizeToCopy = FMath::Min(InstanceData.PlaybackData.Num() * InstanceData.PlaybackData.GetTypeSize(), InstancedMeshComponent->PerInstanceSMCustomData.Num() * InstancedMeshComponent->PerInstanceSMCustomData.GetTypeSize());
-	FMemory::Memcpy(InstancedMeshComponent->PerInstanceSMCustomData.GetData(), InstanceData.PlaybackData.GetData(), CustomDataSizeToCopy);
+	// Clear data.
+	InstancedMeshComponent->ClearInstances();
+
+	if (!NumInstances)
+	{
+		return false;
+	}
+
+	// Allocate Transforms
+	TArray<FTransform> InstanceTransforms;
+	InstanceTransforms.AddDefaulted(NumInstances);
 
-	int32 TransformsToCopy = FMath::Min(InstancedMeshComponent->GetNumRenderInstances(), InstanceData.StaticMeshInstanceData.Num());
-	InstancedMeshComponent->BatchUpdateInstancesData(0, TransformsToCopy, InstanceData.StaticMeshInstanceData.GetData(), true, false);
+	// Set Custom Data Length
+	const SIZE_T DataSize = bAutoPlay ? sizeof(FAnimToTextureAutoPlayData) : sizeof(FAnimToTextureFrameData);
+	InstancedMeshComponent->NumCustomDataFloats = DataSize / sizeof(float);
+
+	// Initizalize Instances
+	InstancedMeshComponent->AddInstances(InstanceTransforms, /*bShouldReturnIndices*/ false, /*bWorldSpace*/ false);
+
+	return true;
 }
 
-void UAnimToTextureInstancePlaybackLibrary::AllocateInstanceData(FAnimToTextureInstanceData& InstanceData, int32 Count)
+bool UAnimToTextureInstancePlaybackLibrary::BatchUpdateInstancesAutoPlayData(UInstancedStaticMeshComponent* InstancedMeshComponent, 
+	const TArray<FAnimToTextureAutoPlayData>& AutoPlayData, const TArray<FMatrix>& Transforms, bool bMarkRenderStateDirty)
 {
-	InstanceData.StaticMeshInstanceData.AddDefaulted(Count);
-	InstanceData.PlaybackData.AddDefaulted(Count);
+	return BatchUpdateInstancesData< FAnimToTextureAutoPlayData>(InstancedMeshComponent, AutoPlayData, Transforms, bMarkRenderStateDirty);
 }
 
-bool UAnimToTextureInstancePlaybackLibrary::UpdateInstanceData(FAnimToTextureInstanceData& InstanceData, int32 InstanceIndex, const FAnimToTextureInstancePlaybackData& PlaybackData, const FTransform& Transform)
+bool UAnimToTextureInstancePlaybackLibrary::BatchUpdateInstancesFrameData(UInstancedStaticMeshComponent* InstancedMeshComponent,
+	const TArray<FAnimToTextureFrameData>& FrameData, const TArray<FMatrix>& Transforms, bool bMarkRenderStateDirty)
 {
-	if (InstanceData.PlaybackData.IsValidIndex(InstanceIndex) && InstanceData.StaticMeshInstanceData.IsValidIndex(InstanceIndex))
-	{
-		InstanceData.PlaybackData[InstanceIndex] = PlaybackData;
-		InstanceData.StaticMeshInstanceData[InstanceIndex].Transform = Transform.ToMatrixWithScale();
-		return true;
-	}
+	return BatchUpdateInstancesData<FAnimToTextureFrameData>(InstancedMeshComponent, FrameData, Transforms, bMarkRenderStateDirty);
+}
 
-	return false;
+bool UAnimToTextureInstancePlaybackLibrary::UpdateInstanceAutoPlayData(UInstancedStaticMeshComponent* InstancedMeshComponent,
+	int32 InstanceIndex, const FAnimToTextureAutoPlayData& AutoPlayData, bool bMarkRenderStateDirty)
+{
+	return UpdateInstanceCustomData<FAnimToTextureAutoPlayData>(InstancedMeshComponent, InstanceIndex, AutoPlayData, bMarkRenderStateDirty);
 }
 
-bool UAnimToTextureInstancePlaybackLibrary::GetInstancePlaybackData(const FAnimToTextureInstanceData& InstanceData, int32 InstanceIndex, FAnimToTextureInstancePlaybackData& InstancePlaybackData)
+bool UAnimToTextureInstancePlaybackLibrary::UpdateInstanceFrameData(UInstancedStaticMeshComponent* InstancedMeshComponent,
+	int32 InstanceIndex, const FAnimToTextureFrameData& FrameData, bool bMarkRenderStateDirty)
 {
-	if (InstanceData.PlaybackData.IsValidIndex(InstanceIndex))
-	{
-		InstancePlaybackData = InstanceData.PlaybackData[InstanceIndex];
-		return true;
-	}
-	
-	return false;
+	return UpdateInstanceCustomData<FAnimToTextureFrameData>(InstancedMeshComponent, InstanceIndex, FrameData, bMarkRenderStateDirty);
 }
 
-bool UAnimToTextureInstancePlaybackLibrary::GetInstanceTransform(const FAnimToTextureInstanceData& InstanceData, int32 InstanceIndex, FTransform& InstanceTransform)
+bool UAnimToTextureInstancePlaybackLibrary::GetAutoPlayDataFromDataAsset(const UAnimToTextureDataAsset* DataAsset, int32 AnimationIndex,
+	FAnimToTextureAutoPlayData& AutoPlayData, float TimeOffset, float PlayRate)
 {
-	if (InstanceData.StaticMeshInstanceData.IsValidIndex(InstanceIndex))
+	if (DataAsset && DataAsset->Animations.IsValidIndex(AnimationIndex))
 	{
-		InstanceTransform = FTransform(InstanceData.StaticMeshInstanceData[InstanceIndex].Transform);
+		// Copy Frame Range
+		const FAnimToTextureAnimInfo& AnimInfo = DataAsset->Animations[AnimationIndex];
+		AutoPlayData.StartFrame = AnimInfo.StartFrame;
+		AutoPlayData.EndFrame = AnimInfo.EndFrame;
+
+		AutoPlayData.TimeOffset = TimeOffset;
+		AutoPlayData.PlayRate = PlayRate;
+
 		return true;
 	}
 
+	// Return default
+	AutoPlayData = FAnimToTextureAutoPlayData();
 	return false;
 }
 
-bool UAnimToTextureInstancePlaybackLibrary::AnimStateFromDataAsset(const UAnimToTextureDataAsset* DataAsset, int32 StateIndex, FAnimToTextureAnimState& AnimState)
+float UAnimToTextureInstancePlaybackLibrary::GetFrame(float Time, float StartFrame, float EndFrame, 
+	float TimeOffset, float PlayRate, float SampleRate)
+{	
+	// Clamp inputs (just in case)
+	Time = FMath::Max(Time, 0.f);
+	TimeOffset = FMath::Max(TimeOffset, 0.f);
+	PlayRate = FMath::Max(PlayRate, 0.f);
+
+	const float Frame = (Time + TimeOffset) * (PlayRate * SampleRate);
+	const float NumFrames = EndFrame - StartFrame + 1.f;
+	return FMath::Fmod(Frame, NumFrames) + StartFrame;
+}
+
+bool UAnimToTextureInstancePlaybackLibrary::GetFrameDataFromDataAsset(const UAnimToTextureDataAsset* DataAsset, int32 AnimationIndex, float Time, FAnimToTextureFrameData& FrameData, float TimeOffset, float PlayRate)
 {
-	if (DataAsset && DataAsset->Animations.IsValidIndex(StateIndex))
+	if (DataAsset && DataAsset->Animations.IsValidIndex(AnimationIndex))
 	{
-		const FAnimInfo& AnimInfo = DataAsset->Animations[StateIndex];
-		AnimState.StartFrame = AnimInfo.AnimStart;
-		AnimState.NumFrames = AnimInfo.NumFrames;
-		AnimState.bLooping = AnimInfo.bLooping;
+		// Copy Frame Range
+		const FAnimToTextureAnimInfo& AnimInfo = DataAsset->Animations[AnimationIndex];
+		
+		FrameData.Frame = GetFrame(Time, AnimInfo.StartFrame, AnimInfo.EndFrame, 
+			TimeOffset, PlayRate, DataAsset->SampleRate);
+
+		// Previous Frame
+		FrameData.PrevFrame = FMath::Clamp(FrameData.Frame - 1, AnimInfo.StartFrame, AnimInfo.EndFrame);
 
 		return true;
 	}
 
-	AnimState = FAnimToTextureAnimState();
+	// Return default
+	FrameData = FAnimToTextureFrameData();
 	return false;
 }
-
diff --git a/Engine/Plugins/Animation/AnimToTexture/Source/AnimToTexture/Private/AnimToTextureModule.cpp b/Engine/Plugins/Animation/AnimToTexture/Source/AnimToTexture/Private/AnimToTextureModule.cpp
new file mode 100644
--- /dev/null
+++ b/Engine/Plugins/Animation/AnimToTexture/Source/AnimToTexture/Private/AnimToTextureModule.cpp
@@ -0,0 +1,24 @@
+// Copyright Epic Games, Inc. All Rights Reserved.
+
+#include "AnimToTextureModule.h"
+
+DEFINE_LOG_CATEGORY(LogAnimToTexture);
+
+#define LOCTEXT_NAMESPACE "FAnimToTextureModule"
+
+void FAnimToTextureModule::StartupModule()
+{
+	// This code will execute after your module is loaded into memory; the exact timing is specified in the .uplugin file per-module
+	
+}
+
+void FAnimToTextureModule::ShutdownModule()
+{
+	// This function may be called during shutdown to clean up your module.  For modules that support dynamic reloading,
+	// we call this function before unloading the module.
+	
+}
+
+#undef LOCTEXT_NAMESPACE
+	
+IMPLEMENT_MODULE(FAnimToTextureModule, AnimToTexture)
\ No newline at end of file
diff --git a/Engine/Plugins/Animation/AnimToTexture/Source/AnimToTexture/Private/EvaluateSequenceAnimInstance.cpp b/Engine/Plugins/Animation/AnimToTexture/Source/AnimToTexture/Private/EvaluateSequenceAnimInstance.cpp
deleted file mode 100644
diff --git a/Engine/Plugins/Animation/AnimToTexture/Source/AnimToTexture/Private/LightweightMontageInstance.cpp b/Engine/Plugins/Animation/AnimToTexture/Source/AnimToTexture/Private/LightweightMontageInstance.cpp
deleted file mode 100644
diff --git a/Engine/Plugins/Animation/AnimToTexture/Source/AnimToTexture/Public/AnimToTexture.h b/Engine/Plugins/Animation/AnimToTexture/Source/AnimToTexture/Public/AnimToTexture.h
deleted file mode 100644
diff --git a/Engine/Plugins/Animation/AnimToTexture/Source/AnimToTexture/Public/AnimToTextureDataAsset.h b/Engine/Plugins/Animation/AnimToTexture/Source/AnimToTexture/Public/AnimToTextureDataAsset.h
--- a/Engine/Plugins/Animation/AnimToTexture/Source/AnimToTexture/Public/AnimToTextureDataAsset.h
+++ b/Engine/Plugins/Animation/AnimToTexture/Source/AnimToTexture/Public/AnimToTextureDataAsset.h
@@ -15,12 +15,16 @@ class UTexture2D;
 
 namespace AnimToTextureParamNames
 {
-	static const FName Animate = TEXT("Animate");
+	static const FName Frame = TEXT("Frame");
+	static const FName AutoPlay = TEXT("AutoPlay");
+	static const FName StartFrame = TEXT("StartFrame");
+	static const FName EndFrame = TEXT("EndFrame");
+	static const FName SampleRate = TEXT("SampleRate");
+	static const FName NumFrames = TEXT("NumFrames");
 	static const FName BoundingBoxMin = TEXT("MinBBox");
 	static const FName BoundingBoxScale = TEXT("SizeBBox");
 	static const FName RowsPerFrame = TEXT("RowsPerFrame");
 	static const FName BoneWeightRowsPerFrame = TEXT("BoneWeightsRowsPerFrame");
-	static const FName NumFrames = TEXT("NumFrames");
 	static const FName VertexPositionTexture = TEXT("PositionTexture");
 	static const FName VertexNormalTexture = TEXT("NormalTexture");
 	static const FName BonePositionTexture = TEXT("BonePositionTexture");
@@ -64,7 +68,7 @@ enum class EAnimToTextureNumBoneInfluences : uint8
 };
 
 USTRUCT(Blueprintable)
-struct FAnimSequenceInfo
+struct FAnimToTextureAnimSequenceInfo
 {
 	GENERATED_BODY()
 
@@ -74,9 +78,6 @@ struct FAnimSequenceInfo
 	UPROPERTY(EditAnywhere, Category = Default, BlueprintReadWrite, meta = (EditCondition = "bEnabled", EditConditionHides))
 	TObjectPtr<UAnimSequence> AnimSequence = nullptr;
 
-	UPROPERTY(EditAnywhere, Category = Default, BlueprintReadWrite, meta = (EditCondition = "bEnabled", EditConditionHides))
-	bool bLooping = true;
-
 	UPROPERTY(EditAnywhere, Category = Default, BlueprintReadWrite, meta = (EditCondition = "bEnabled", EditConditionHides))
 	bool bUseCustomRange = false;
 
@@ -89,18 +90,16 @@ struct FAnimSequenceInfo
 };
 
 USTRUCT(Blueprintable)
-struct FAnimInfo
+struct FAnimToTextureAnimInfo
 {
 	GENERATED_BODY()
-
+	
 	UPROPERTY(VisibleAnywhere, Category = Default, BlueprintReadOnly)
-	int32 NumFrames = 0;
+	int32 StartFrame = 0;
 
 	UPROPERTY(VisibleAnywhere, Category = Default, BlueprintReadOnly)
-	int32 AnimStart = 0;
+	int32 EndFrame = 0;
 
-	UPROPERTY(EditAnywhere, Category = Default, BlueprintReadOnly)
-	bool bLooping = true;
 };
 
 UCLASS(Blueprintable, BlueprintType)
@@ -227,7 +226,7 @@ public:
 	float SampleRate = 30.0f;
 
 	UPROPERTY(EditAnywhere, BlueprintReadWrite, Category = "Animation")
-	TArray<FAnimSequenceInfo> AnimSequences;
+	TArray<FAnimToTextureAnimSequenceInfo> AnimSequences;
 
 	// ------------------------------------------------------
 	// Info
@@ -258,8 +257,12 @@ public:
 	FVector BoneSizeBBox;
 
 	UPROPERTY(VisibleAnywhere, BlueprintReadOnly, Category = "Info")
-	TArray<FAnimInfo> Animations;
+	TArray<FAnimToTextureAnimInfo> Animations;
 
+	/* Finds AnimSequence Index in the Animations Array. 
+	*  Only Enabled elements are returned.
+	*  Returns -1 if not found.
+	*/
 	UFUNCTION(BlueprintCallable, Category = Default)
 	int32 GetIndexFromAnimSequence(const UAnimSequence* Sequence);
 
diff --git a/Engine/Plugins/Animation/AnimToTexture/Source/AnimToTexture/Public/AnimToTextureInstancePlaybackHelpers.h b/Engine/Plugins/Animation/AnimToTexture/Source/AnimToTexture/Public/AnimToTextureInstancePlaybackHelpers.h
--- a/Engine/Plugins/Animation/AnimToTexture/Source/AnimToTexture/Public/AnimToTextureInstancePlaybackHelpers.h
+++ b/Engine/Plugins/Animation/AnimToTexture/Source/AnimToTexture/Public/AnimToTextureInstancePlaybackHelpers.h
@@ -7,66 +7,56 @@
 
 #include "AnimToTextureInstancePlaybackHelpers.generated.h"
 
-// Use floats to match custom floats of instanced static mesh
-// We could pack a float w/ more parameters if desired
-USTRUCT(BlueprintType)
-struct FAnimToTextureAnimState
-{
-	GENERATED_USTRUCT_BODY()
-
-	UPROPERTY(EditAnywhere, BlueprintReadWrite, Category = "AnimToTexture")
-	float StartFrame = 0.0f;
-
-	UPROPERTY(EditAnywhere, BlueprintReadWrite, Category = "AnimToTexture")
-	float NumFrames = 60.0f;
-
-	UPROPERTY(EditAnywhere, BlueprintReadWrite, Category = "AnimToTexture")
-	float PlayRate = 1.0f;
-
-	UPROPERTY(EditAnywhere, BlueprintReadWrite, Category = "AnimToTexture")
-	float bLooping = 1.0f; 
-
-	UPROPERTY(EditAnywhere, BlueprintReadWrite, Category = "AnimToTexture")
-	float GlobalStartTime = 0.0f; 
-};
+class UAnimToTextureDataAsset;
 
 USTRUCT(BlueprintType)
-struct FAnimToTextureInstancePlaybackData
+struct FAnimToTextureFrameData
 {
 	GENERATED_USTRUCT_BODY()
 
-	// Store prev state to allow blending of prev->current state in material
-	// Uncomment this if we start blending states
-	//UPROPERTY(EditAnywhere, BlueprintReadWrite, Category = "AnimToTexture")
-	//FAnimToTextureAnimState PrevState;
-
-	UPROPERTY(EditAnywhere, BlueprintReadWrite, Category = "AnimToTexture")
-	FAnimToTextureAnimState CurrentState;
+	/**
+	* Frame to be played
+	*/
+	UPROPERTY(EditAnywhere, BlueprintReadWrite, Category = "AnimToTexture|Playback")
+	float Frame = 0.0f;
+
+	/**
+	* Previous Frame (this is needeed for motion blur)
+	*/
+	UPROPERTY(EditAnywhere, BlueprintReadWrite, Category = "AnimToTexture|Playback")
+	float PrevFrame = 0.0f;
 };
 
 USTRUCT(BlueprintType)
-struct FAnimToTextureAnimationSyncData
+struct FAnimToTextureAutoPlayData
 {
 	GENERATED_USTRUCT_BODY()
 
-	// The time used for sync when transitioning from skeletal mesh to material animated static mesh.
-	// World real time at the time of the transition
-	float SyncTime;
-};
+	/**
+	* Adds offset to time
+	*/
+	UPROPERTY(EditAnywhere, BlueprintReadWrite, Category = "AnimToTexture|Playback")
+	float TimeOffset = 0.0f;
 
-USTRUCT(BlueprintType)
-struct FAnimToTextureInstanceData
-{
-	GENERATED_USTRUCT_BODY()
+	/**
+	* Rate for increasing and decreasing speed.
+	*/
+	UPROPERTY(EditAnywhere, BlueprintReadWrite, Category = "AnimToTexture|Playback")
+	float PlayRate = 1.0f;
 
-	UPROPERTY()
-	TArray<FAnimToTextureInstancePlaybackData> PlaybackData;
+	/**
+	* Starting frame for animation. 
+	*/
+	UPROPERTY(EditAnywhere, BlueprintReadWrite, Category = "AnimToTexture|Playback")
+	float StartFrame = 0.0f;
 
-	UPROPERTY()
-	TArray<FInstancedStaticMeshInstanceData> StaticMeshInstanceData;
+	/**
+	* Last frame of animation
+	*/
+	UPROPERTY(EditAnywhere, BlueprintReadWrite, Category = "AnimToTexture|Playback")
+	float EndFrame = 1.0f;
 };
 
-class UAnimToTextureDataAsset;
 
 UCLASS()
 class ANIMTOTEXTURE_API UAnimToTextureInstancePlaybackLibrary : public UBlueprintFunctionLibrary
@@ -74,25 +64,150 @@ class ANIMTOTEXTURE_API UAnimToTextureInstancePlaybackLibrary : public UBlueprin
 	GENERATED_BODY()
 
 public:
+	
+	/**
+	* Adds Instances and allocates the necessary CustomData.
+	* @param bAutoPlay: if true, it will allocate the required CustomData for working with AutoPlayData. If false, FrameData will be allocated instead./
+	*/
+	UFUNCTION(BlueprintCallable, Category = "AnimToTexture|Playback")
+	static bool SetupInstancedMeshComponent(UInstancedStaticMeshComponent* InstancedMeshComponent, int32 NumInstances, bool bAutoPlay);
+
+	/**
+	* Updates all instances with the given Transforms and AutoPlayData
+	* @param bMarkRenderStateDirty: if true, the change should be visible immediatelly. If you are updating many instances you should only set this to true for the last instance
+	*/
+	UFUNCTION(BlueprintCallable, Category = "AnimToTexture|Playback")
+	static bool BatchUpdateInstancesAutoPlayData(UInstancedStaticMeshComponent* InstancedMeshComponent,
+		const TArray<FAnimToTextureAutoPlayData>& AutoPlayData, const TArray<FMatrix>& Transforms, bool bMarkRenderStateDirty=true);
 
+	/**
+	* Updates all instances with the given Transforms and FrameData
+	* @param bMarkRenderStateDirty: if true, the change should be visible immediatelly. If you are updating many instances you should only set this to true for the last instance
+	*/
 	UFUNCTION(BlueprintCallable, Category = "AnimToTexture|Playback")
-	static void SetupInstancedMeshComponent(UInstancedStaticMeshComponent* InstancedMeshComponent, UPARAM(ref) FAnimToTextureInstanceData& InstanceData, int32 NumInstances);
+	static bool BatchUpdateInstancesFrameData(UInstancedStaticMeshComponent* InstancedMeshComponent,
+		const TArray<FAnimToTextureFrameData>& FrameData, const TArray<FMatrix>& Transforms, bool bMarkRenderStateDirty=true);
 
+	/**
+	* Updates a single instance with given AutoPlayData
+	* @param bMarkRenderStateDirty: if true, the change should be visible immediatelly. If you are updating many instances you should only set this to true for the last instance
+	*/
 	UFUNCTION(BlueprintCallable, Category = "AnimToTexture|Playback")
-	static void BatchUpdateInstancedMeshComponent(UInstancedStaticMeshComponent* InstancedMeshComponent, UPARAM(ref) FAnimToTextureInstanceData& InstanceData);
+	static bool UpdateInstanceAutoPlayData(UInstancedStaticMeshComponent* InstancedMeshComponent, int32 InstanceIndex, 
+		const FAnimToTextureAutoPlayData& AutoPlayData, bool bMarkRenderStateDirty=true);
 
+	/**
+	* * Updates a single instance with given FrameData
+	* @param bMarkRenderStateDirty: if true, the change should be visible immediatelly. If you are updating many instances you should only set this to true for the last instance
+	*/
 	UFUNCTION(BlueprintCallable, Category = "AnimToTexture|Playback")
-	static void AllocateInstanceData(UPARAM(ref) FAnimToTextureInstanceData& InstanceData, int32 Count);
+	static bool UpdateInstanceFrameData(UInstancedStaticMeshComponent* InstancedMeshComponent, int32 InstanceIndex, 
+		const FAnimToTextureFrameData& FrameData, bool bMarkRenderStateDirty=true);
 
+	/**
+	* Returns an AutoPlayData with the Start and End Frame for the given AnimationIndex.
+	* If AnimationIndex is out of range, false will be returned.
+	*/
 	UFUNCTION(BlueprintCallable, Category = "AnimToTexture|Playback")
-	static bool UpdateInstanceData(UPARAM(ref) FAnimToTextureInstanceData& InstanceData, int32 InstanceIndex, const FAnimToTextureInstancePlaybackData& PlaybackData, const FTransform& Transform);
+	static bool GetAutoPlayDataFromDataAsset(const UAnimToTextureDataAsset* DataAsset, int32 AnimationIndex,
+		FAnimToTextureAutoPlayData& AutoPlayData, float TimeOffset = 0.f, float PlayRate = 1.f);
 
 	UFUNCTION(BlueprintCallable, Category = "AnimToTexture|Playback")
-	static bool GetInstancePlaybackData(UPARAM(ref) const FAnimToTextureInstanceData& InstanceData, int32 InstanceIndex, FAnimToTextureInstancePlaybackData& InstancePlaybackData);
+	static float GetFrame(float Time, float StartFrame, float EndFrame, 
+		float TimeOffset = 0.f, float PlayRate = 1.f, float SampleRate = 30.f);
 
 	UFUNCTION(BlueprintCallable, Category = "AnimToTexture|Playback")
-	static bool GetInstanceTransform(UPARAM(ref) const FAnimToTextureInstanceData& InstanceData, int32 InstanceIndex, FTransform& InstanceTransform);
+	static bool GetFrameDataFromDataAsset(const UAnimToTextureDataAsset* DataAsset, int32 AnimationIndex, float Time,
+		FAnimToTextureFrameData& AutoPlayData, float TimeOffset = 0.f, float PlayRate = 1.f);
+	
+private:
+	
+	template <class DataType>
+	static bool BatchUpdateInstancesData(UInstancedStaticMeshComponent* InstancedMeshComponent,
+		const TArray<DataType>& Data, const TArray<FMatrix>& Transforms, bool bMarkRenderStateDirty = true);
+	
+	template <class DataType>
+	static bool UpdateInstanceCustomData(UInstancedStaticMeshComponent* InstancedMeshComponent, int32 InstanceIndex,
+		const DataType& Data, bool bMarkRenderStateDirty = true);
+
+	template <class DataType>
+	static void GetData(const DataType& Data, TArray<float>& CustomData);
+
+	template <class DataType>
+	static void GetCustomData(const TArray<DataType>& Data, TArray<float>& CustomData);
 
-	UFUNCTION(BlueprintPure, BlueprintCallable, Category = "AnimToTexture|Playback")
-	static bool AnimStateFromDataAsset(const UAnimToTextureDataAsset* DataAsset, int32 StateIndex, FAnimToTextureAnimState& AnimState);
 };
+
+
+template <class DataType>
+bool UAnimToTextureInstancePlaybackLibrary::BatchUpdateInstancesData(UInstancedStaticMeshComponent* InstancedMeshComponent,
+	const TArray<DataType>& Data, const TArray<FMatrix>& Transforms, bool bMarkRenderStateDirty)
+{
+	if (!InstancedMeshComponent)
+	{
+		return false;
+	}
+
+	// Get Number of Instances
+	const int32 NumInstances = InstancedMeshComponent->GetNumRenderInstances();
+
+	// Check number of instances
+	if (NumInstances &&
+		NumInstances == Data.Num() &&
+		NumInstances == Transforms.Num())
+	{
+		// Check data size
+		if (InstancedMeshComponent->NumCustomDataFloats == Data.GetTypeSize() / sizeof(float))
+		{
+			SIZE_T CustomDataSizeToCopy = InstancedMeshComponent->PerInstanceSMCustomData.Num() * InstancedMeshComponent->PerInstanceSMCustomData.GetTypeSize();
+			FMemory::Memcpy(InstancedMeshComponent->PerInstanceSMCustomData.GetData(), Data.GetData(), CustomDataSizeToCopy);
+
+			// FInstancedStaticMeshInstanceData is not exposed to blueprints, so we are duplicating the data here :/
+			TArray<FInstancedStaticMeshInstanceData> InstanceData;
+			InstanceData.AddUninitialized(NumInstances);
+			FMemory::Memcpy(InstanceData.GetData(), Transforms.GetData(), NumInstances * InstanceData.GetTypeSize());
+
+			// Batch Update
+			return InstancedMeshComponent->BatchUpdateInstancesData(0, NumInstances, InstanceData.GetData(), bMarkRenderStateDirty, false);
+		}
+	}
+
+	return false;
+}
+
+template <class DataType>
+void UAnimToTextureInstancePlaybackLibrary::GetData(const DataType& Data, TArray<float>& CustomData)
+{
+	const SIZE_T DataSize = sizeof(Data);
+	CustomData.SetNumUninitialized(DataSize / sizeof(float));
+
+	FMemory::Memcpy(CustomData.GetData(), &Data, DataSize);
+}
+
+template <class DataType>
+void UAnimToTextureInstancePlaybackLibrary::GetCustomData(const TArray<DataType>& Data, TArray<float>& CustomData)
+{
+	const int32 NumCustomDataFloats = Data.GetTypeSize() / sizeof(float);
+	CustomData.SetNumUninitialized(Data.Num() * NumCustomDataFloats);
+	
+	FMemory::Memcpy(CustomData.GetData(), Data.GetData(), Data.Num() * Data.GetTypeSize());
+}
+
+template <class DataType>
+bool UAnimToTextureInstancePlaybackLibrary::UpdateInstanceCustomData(UInstancedStaticMeshComponent* InstancedMeshComponent, int32 InstanceIndex,
+	const DataType& Data, bool bMarkRenderStateDirty)
+{
+	if (InstancedMeshComponent && InstancedMeshComponent->IsValidInstance(InstanceIndex))
+	{
+		// Check if valid size
+		const int32 NumCustomDataFloats = sizeof(Data) / sizeof(float);
+		if (InstancedMeshComponent->NumCustomDataFloats == NumCustomDataFloats)
+		{
+			TArray<float> CustomData;
+			GetData(Data, CustomData);
+			return InstancedMeshComponent->SetCustomData(InstanceIndex, CustomData, bMarkRenderStateDirty);
+		}
+	}
+
+	return false;
+}
diff --git a/Engine/Plugins/Animation/AnimToTexture/Source/AnimToTexture/Public/AnimToTextureModule.h b/Engine/Plugins/Animation/AnimToTexture/Source/AnimToTexture/Public/AnimToTextureModule.h
new file mode 100644
--- /dev/null
+++ b/Engine/Plugins/Animation/AnimToTexture/Source/AnimToTexture/Public/AnimToTextureModule.h
@@ -0,0 +1,16 @@
+// Copyright Epic Games, Inc. All Rights Reserved.
+
+#pragma once
+
+#include "Modules/ModuleManager.h"
+
+DECLARE_LOG_CATEGORY_EXTERN(LogAnimToTexture, Log, All);
+
+class FAnimToTextureModule : public IModuleInterface
+{
+public:
+
+	/** IModuleInterface implementation */
+	virtual void StartupModule() override;
+	virtual void ShutdownModule() override;
+};
diff --git a/Engine/Plugins/Animation/AnimToTexture/Source/AnimToTexture/Public/EvaluateSequenceAnimInstance.h b/Engine/Plugins/Animation/AnimToTexture/Source/AnimToTexture/Public/EvaluateSequenceAnimInstance.h
deleted file mode 100644
diff --git a/Engine/Plugins/Animation/AnimToTexture/Source/AnimToTexture/Public/LightweightMontageInstance.h b/Engine/Plugins/Animation/AnimToTexture/Source/AnimToTexture/Public/LightweightMontageInstance.h
deleted file mode 100644
diff --git a/Engine/Plugins/Animation/AnimToTexture/Source/AnimToTextureEditor/Private/AnimToTextureBPLibrary.cpp b/Engine/Plugins/Animation/AnimToTexture/Source/AnimToTextureEditor/Private/AnimToTextureBPLibrary.cpp
--- a/Engine/Plugins/Animation/AnimToTexture/Source/AnimToTextureEditor/Private/AnimToTextureBPLibrary.cpp
+++ b/Engine/Plugins/Animation/AnimToTexture/Source/AnimToTextureEditor/Private/AnimToTextureBPLibrary.cpp
@@ -1,14 +1,14 @@
 // Copyright Epic Games, Inc. All Rights Reserved.
 
 #include "AnimToTextureBPLibrary.h"
-#include "AnimToTexture.h"
+#include "AnimToTextureEditorModule.h"
 #include "AnimToTextureUtils.h"
 #include "AnimToTextureSkeletalMesh.h"
-#include "EvaluateSequenceAnimInstance.h"
 
 #include "LevelEditor.h"
 #include "RawMesh.h"
 #include "MeshUtilities.h"
+#include "Modules/ModuleManager.h"
 #include "Engine/SkeletalMesh.h"
 #include "Components/SkeletalMeshComponent.h"
 #include "Animation/Skeleton.h"
@@ -42,35 +42,35 @@ void UAnimToTextureBPLibrary::AnimationToTexture(UAnimToTextureDataAsset* DataAs
 	// Check StaticMesh
 	if (!DataAsset->GetStaticMesh())
 	{
-		UE_LOG(LogTemp, Warning, TEXT("Invalid StaticMesh"));
+		UE_LOG(LogAnimToTextureEditor, Warning, TEXT("Invalid StaticMesh"));
 		return;
 	}
 
 	// Check SkeletalMesh
 	if (!DataAsset->GetSkeletalMesh())
 	{
-		UE_LOG(LogTemp, Warning, TEXT("Invalid SkeletalMesh"));
+		UE_LOG(LogAnimToTextureEditor, Warning, TEXT("Invalid SkeletalMesh"));
 		return;
 	}
 
 	// Check Skeleton
 	if (!DataAsset->GetSkeletalMesh()->GetSkeleton())
 	{
-		UE_LOG(LogTemp, Warning, TEXT("Invalid SkeletalMesh. No valid Skeleton found"));
+		UE_LOG(LogAnimToTextureEditor, Warning, TEXT("Invalid SkeletalMesh. No valid Skeleton found"));
 		return;
 	}
 
 	// Check StaticMesh LOD
 	if (!DataAsset->GetStaticMesh()->IsSourceModelValid(DataAsset->StaticLODIndex))
 	{
-		UE_LOG(LogTemp, Warning, TEXT("Invalid StaticMesh LOD Index"));
+		UE_LOG(LogAnimToTextureEditor, Warning, TEXT("Invalid StaticMesh LOD Index: %i"), DataAsset->StaticLODIndex);
 		return;
 	}
 
 	// Check SkeletalMesh LOD
 	if (!DataAsset->GetSkeletalMesh()->IsValidLODIndex(DataAsset->SkeletalLODIndex))
 	{
-		UE_LOG(LogTemp, Warning, TEXT("Invalid SkeletalMesh LOD Index"));
+		UE_LOG(LogAnimToTextureEditor, Warning, TEXT("Invalid SkeletalMesh LOD Index: %i"), DataAsset->SkeletalLODIndex);
 		return;
 	}
 
@@ -84,13 +84,13 @@ void UAnimToTextureBPLibrary::AnimationToTexture(UAnimToTextureDataAsset* DataAs
 		}
 		else
 		{
-			UE_LOG(LogTemp, Warning, TEXT("Invalid Socket: %s"), *DataAsset->AttachToSocket.ToString());
+			UE_LOG(LogAnimToTextureEditor, Warning, TEXT("Invalid Socket: %s"), *DataAsset->AttachToSocket.ToString());
 			return;
 		}
 	}
 	if (bValidSocket && DataAsset->Mode == EAnimToTextureMode::Vertex)
 	{
-		UE_LOG(LogTemp, Warning, TEXT("Unable to use Socket in Vertex Mode. Use Bone Mode instead."));
+		UE_LOG(LogAnimToTextureEditor, Warning, TEXT("Unable to use Socket in Vertex Mode. Use Bone Mode instead."));
 		return;
 	}
 
@@ -99,10 +99,34 @@ void UAnimToTextureBPLibrary::AnimationToTexture(UAnimToTextureDataAsset* DataAs
 	if (SourceModel.BuildSettings.bGenerateLightmapUVs &&
 		SourceModel.BuildSettings.DstLightmapIndex == DataAsset->UVChannel)
 	{
-		UE_LOG(LogTemp, Warning, TEXT("Invalid UVChannel: %i. Already used by LightMap"), DataAsset->UVChannel);
+		UE_LOG(LogAnimToTextureEditor, Warning, TEXT("Invalid UVChannel: %i. Already used by LightMap"), DataAsset->UVChannel);
 		return;
 	}
 
+	// Check Animations
+	int32 NumAnimations = 0;
+	for (const FAnimToTextureAnimSequenceInfo& AnimSequenceInfo : DataAsset->AnimSequences)
+	{
+		const UAnimSequence* AnimSequence = AnimSequenceInfo.AnimSequence;
+		if (AnimSequenceInfo.bEnabled && AnimSequence)
+		{
+			// Check Frame Range
+			if (AnimSequenceInfo.bUseCustomRange && 
+				(AnimSequenceInfo.EndFrame - AnimSequenceInfo.StartFrame) <= 0 )
+			{ 
+				UE_LOG(LogAnimToTextureEditor, Warning, TEXT("Invalid Custom Range for AnimSequence: %s"), *AnimSequence->GetName());
+				return;
+			}
+
+			NumAnimations++;
+		}
+	}
+	if (!NumAnimations)
+	{
+		UE_LOG(LogAnimToTextureEditor, Warning, TEXT("No Animations found"));
+		return;
+	};
+
 	// ---------------------------------------------------------------------------		
 	// Get Meshes Vertices and Mapping.
 	// NOTE: We need to create a Mapping between the StaticMesh and the SkeletalMesh
@@ -136,14 +160,14 @@ void UAnimToTextureBPLibrary::AnimationToTexture(UAnimToTextureDataAsset* DataAs
 	check(Actor);
 
 	// Create Temp SkeletalMesh Component
-	USkeletalMeshComponent* MeshComponent = NewObject<USkeletalMeshComponent>(Actor);
-	check(MeshComponent);
-	MeshComponent->SetSkeletalMesh(DataAsset->GetSkeletalMesh());
-	MeshComponent->SetForcedLOD(1); // Force to LOD0;
-	MeshComponent->SetAnimationMode(EAnimationMode::AnimationSingleNode);
-	MeshComponent->SetUpdateAnimationInEditor(true);
-	MeshComponent->VisibilityBasedAnimTickOption = EVisibilityBasedAnimTickOption::AlwaysTickPoseAndRefreshBones;
-	MeshComponent->RegisterComponent();
+	USkeletalMeshComponent* SkeletalMeshComponent = NewObject<USkeletalMeshComponent>(Actor);
+	check(SkeletalMeshComponent);
+	SkeletalMeshComponent->SetSkeletalMesh(DataAsset->GetSkeletalMesh());
+	SkeletalMeshComponent->SetForcedLOD(1); // Force to LOD0;
+	SkeletalMeshComponent->SetAnimationMode(EAnimationMode::AnimationSingleNode);
+	SkeletalMeshComponent->SetUpdateAnimationInEditor(true);
+	SkeletalMeshComponent->VisibilityBasedAnimTickOption = EVisibilityBasedAnimTickOption::AlwaysTickPoseAndRefreshBones;
+	SkeletalMeshComponent->RegisterComponent();
 
 	// ---------------------------------------------------------------------------
 	// Get Reference Skeleton Transforms
@@ -159,18 +183,18 @@ void UAnimToTextureBPLibrary::AnimationToTexture(UAnimToTextureDataAsset* DataAs
 	if (DataAsset->Mode == EAnimToTextureMode::Bone)
 	{
 		// Gets Ref Bone Position and Rotations.
-		NumBones = GetRefBonePositionsAndRotations(MeshComponent,
+		NumBones = GetRefBonePositionsAndRotations(SkeletalMeshComponent,
 			BoneRefPositions, BoneRefRotations);
 
 		// NOTE: there is a limitation with the number of bones atm.
 		if (NumBones > 256)
 		{
-			UE_LOG(LogTemp, Warning, TEXT("Invalid Number of Bones. There is a maximum of 256 bones"))
+			UE_LOG(LogAnimToTextureEditor, Warning, TEXT("Invalid Number of Bones. There is a maximum of 256 bones"))
 			return;
 		}
 
 		// Get Bone Names (no virtual)
-		GetBoneNames(MeshComponent->GetSkeletalMeshAsset(), BoneNames);
+		GetBoneNames(SkeletalMeshComponent->GetSkeletalMeshAsset(), BoneNames);
 
 		// Make sure array sizes are correct.
 		check(BoneNames.Num() == NumBones);
@@ -178,7 +202,7 @@ void UAnimToTextureBPLibrary::AnimationToTexture(UAnimToTextureDataAsset* DataAs
 		// Check if Socket is in BoneNames
 		if (bValidSocket && !BoneNames.Find(DataAsset->AttachToSocket, SocketIndex))
 		{
-			UE_LOG(LogTemp, Warning, TEXT("Socket: %s not found in Raw Bone List"), *DataAsset->AttachToSocket.ToString());
+			UE_LOG(LogAnimToTextureEditor, Warning, TEXT("Socket: %s not found in Raw Bone List"), *DataAsset->AttachToSocket.ToString());
 			return;
 		}
 
@@ -198,7 +222,7 @@ void UAnimToTextureBPLibrary::AnimationToTexture(UAnimToTextureDataAsset* DataAs
 	//
 	const float SampleInterval = 1.f / DataAsset->SampleRate;
 
-	for (const FAnimSequenceInfo& AnimSequenceInfo : DataAsset->AnimSequences)
+	for (const FAnimToTextureAnimSequenceInfo& AnimSequenceInfo : DataAsset->AnimSequences)
 	{
 		UAnimSequence* AnimSequence = AnimSequenceInfo.AnimSequence;
 
@@ -208,15 +232,15 @@ void UAnimToTextureBPLibrary::AnimationToTexture(UAnimToTextureDataAsset* DataAs
 		}
 		
 		// Make sure SkeletalMesh is compatible with AnimSequence
-		if (!MeshComponent->GetSkeletalMeshAsset()->GetSkeleton()->IsCompatibleForEditor(AnimSequence->GetSkeleton()))
+		if (!SkeletalMeshComponent->GetSkeletalMeshAsset()->GetSkeleton()->IsCompatibleForEditor(AnimSequence->GetSkeleton()))
 		{
-			UE_LOG(LogTemp, Warning, TEXT("Invalid AnimSequence: %s for given SkeletalMesh: %s"), *AnimSequence->GetFName().ToString(), *MeshComponent->GetSkeletalMeshAsset()->GetFName().ToString());
+			UE_LOG(LogAnimToTextureEditor, Warning, TEXT("Invalid AnimSequence: %s for given SkeletalMesh: %s"), *AnimSequence->GetFName().ToString(), *SkeletalMeshComponent->GetSkeletalMeshAsset()->GetFName().ToString());
 			continue;
 		}
 		// Set AnimSequence
 		else
 		{			
-			MeshComponent->SetAnimation(AnimSequence);
+			SkeletalMeshComponent->SetAnimation(AnimSequence);
 		}
 
 		// -----------------------------------------------------------------------------------
@@ -255,12 +279,12 @@ void UAnimToTextureBPLibrary::AnimationToTexture(UAnimToTextureDataAsset* DataAs
 			SampleIndex++;
 
 			// Go To Time
-			MeshComponent->SetPosition(Time);
+			SkeletalMeshComponent->SetPosition(Time);
 
 			// Update SkelMesh Animation.
-			MeshComponent->TickAnimation(0.f, false /*bNeedsValidRootMotion*/);
-			// MeshComponent->TickComponent(0.f, ELevelTick::LEVELTICK_All, nullptr);
-			MeshComponent->RefreshBoneTransforms(nullptr /*TickFunction*/);
+			SkeletalMeshComponent->TickAnimation(0.f, false /*bNeedsValidRootMotion*/);
+			// SkeletalMeshComponent->TickComponent(0.f, ELevelTick::LEVELTICK_All, nullptr);
+			SkeletalMeshComponent->RefreshBoneTransforms(nullptr /*TickFunction*/);
 			
 			// ---------------------------------------------------------------------------
 			// Store Vertex Deltas & Normals.
@@ -270,7 +294,7 @@ void UAnimToTextureBPLibrary::AnimationToTexture(UAnimToTextureDataAsset* DataAs
 				TArray<FVector3f> VertexFrameDeltas;
 				TArray<FVector3f> VertexFrameNormals;
 				
-				GetVertexDeltasAndNormals(MeshComponent, DataAsset->SkeletalLODIndex,
+				GetVertexDeltasAndNormals(SkeletalMeshComponent, DataAsset->SkeletalLODIndex,
 					Vertices, Normals, Mapping,
 					RootTransform,
 					VertexFrameDeltas, VertexFrameNormals);
@@ -288,7 +312,7 @@ void UAnimToTextureBPLibrary::AnimationToTexture(UAnimToTextureDataAsset* DataAs
 				TArray<FVector3f> BoneFramePositions;
 				TArray<FVector4> BoneFrameRotations;
 
-				GetBonePositionsAndRotations(MeshComponent, BoneRefPositions,
+				GetBonePositionsAndRotations(SkeletalMeshComponent, BoneRefPositions,
 					BoneFramePositions, BoneFrameRotations);
 
 				BonePositions.Append(BoneFramePositions);
@@ -299,19 +323,18 @@ void UAnimToTextureBPLibrary::AnimationToTexture(UAnimToTextureDataAsset* DataAs
 		} // End Frame
 
 		// Store Anim Info Data
-		FAnimInfo AnimInfo;
-		AnimInfo.NumFrames = SampleIndex;
-		AnimInfo.AnimStart = DataAsset->NumFrames;
-		AnimInfo.bLooping = AnimSequenceInfo.bLooping;
+		FAnimToTextureAnimInfo AnimInfo;
+		AnimInfo.StartFrame = DataAsset->NumFrames;
+		AnimInfo.EndFrame = DataAsset->NumFrames + EndFrame;
 		DataAsset->Animations.Add(AnimInfo);
 
 		// Accumulate Frames
-		DataAsset->NumFrames += AnimInfo.NumFrames;
+		DataAsset->NumFrames += SampleIndex;
 	} // End Anim
 		
 	// Destroy Temp Component
-	MeshComponent->UnregisterComponent();
-	MeshComponent->DestroyComponent();
+	SkeletalMeshComponent->UnregisterComponent();
+	SkeletalMeshComponent->DestroyComponent();
 
 	Actor->Destroy();
 	
@@ -332,7 +355,7 @@ void UAnimToTextureBPLibrary::AnimationToTexture(UAnimToTextureDataAsset* DataAs
 								Height, Width, DataAsset->VertexRowsPerFrame, 
 								DataAsset->MaxHeight, DataAsset->MaxWidth, DataAsset->bEnforcePowerOfTwo))
 		{
-			UE_LOG(LogTemp, Warning, TEXT("Vertex Animation data cannot be fit in a %ix%i texture."), DataAsset->MaxHeight, DataAsset->MaxWidth);
+			UE_LOG(LogAnimToTextureEditor, Warning, TEXT("Vertex Animation data cannot be fit in a %ix%i texture."), DataAsset->MaxHeight, DataAsset->MaxWidth);
 			return;
 		}
 
@@ -378,7 +401,7 @@ void UAnimToTextureBPLibrary::AnimationToTexture(UAnimToTextureDataAsset* DataAs
 								Height, Width, DataAsset->BoneRowsPerFrame, 
 								DataAsset->MaxHeight, DataAsset->MaxWidth, DataAsset->bEnforcePowerOfTwo))
 		{
-			UE_LOG(LogTemp, Warning, TEXT("Bone Animation data cannot be fit in a %ix%i texture."), DataAsset->MaxHeight, DataAsset->MaxWidth);
+			UE_LOG(LogAnimToTextureEditor, Warning, TEXT("Bone Animation data cannot be fit in a %ix%i texture."), DataAsset->MaxHeight, DataAsset->MaxWidth);
 			return;
 		}
 
@@ -412,7 +435,7 @@ void UAnimToTextureBPLibrary::AnimationToTexture(UAnimToTextureDataAsset* DataAs
 								Height, Width, DataAsset->BoneWeightRowsPerFrame, 
 								DataAsset->MaxHeight, DataAsset->MaxWidth, DataAsset->bEnforcePowerOfTwo))
 		{
-			UE_LOG(LogTemp, Warning, TEXT("Weights Data cannot be fit in a %ix%i texture."), DataAsset->MaxHeight, DataAsset->MaxWidth);
+			UE_LOG(LogAnimToTextureEditor, Warning, TEXT("Weights Data cannot be fit in a %ix%i texture."), DataAsset->MaxHeight, DataAsset->MaxWidth);
 			return;
 		}
 
@@ -482,12 +505,18 @@ void UAnimToTextureBPLibrary::GetVertexDeltasAndNormals(const USkeletalMeshCompo
 	OutVertexNormals.Reset();
 
 	check(Vertices.Num() == Mapping.Num());
+	const int32 NumVertices = Vertices.Num();
 
+	// Get Deformed vertices at current frame
 	TArray<FVector3f> SkinnedVertices;
 	GetSkinnedVertices(SkeletalMeshComponent, LODIndex, SkinnedVertices);
 
+	// Allocate
+	OutVertexDeltas.SetNumUninitialized(NumVertices);
+	OutVertexNormals.SetNumUninitialized(NumVertices);
+
 	// Loop thru static vertices and find the mapped SkeletalMesh Vertex
-	for (int32 VertexIndex = 0; VertexIndex < Vertices.Num(); VertexIndex++)
+	for (int32 VertexIndex = 0; VertexIndex < NumVertices; VertexIndex++)
 	{
 		// Get Static/SkeletalMesh Mapping
 		const FVertexToMeshMapping& VertexMapping = Mapping[VertexIndex];
@@ -506,8 +535,8 @@ void UAnimToTextureBPLibrary::GetVertexDeltasAndNormals(const USkeletalMeshCompo
 		const FVector3f TransformedVertexDelta = ((FVector3f)RootTransform.TransformPosition((FVector)VertexPosition)) - Vertices[VertexIndex];
 		const FVector3f TransformedVertexNormal = (FVector3f)RootTransform.TransformVector((FVector)VertexNormal);
 		
-		OutVertexDeltas.Add(TransformedVertexDelta);
-		OutVertexNormals.Add(TransformedVertexNormal);
+		OutVertexDeltas[VertexIndex] = TransformedVertexDelta;
+		OutVertexNormals[VertexIndex] = TransformedVertexNormal;
 	}
 }
 
@@ -550,6 +579,10 @@ int32 UAnimToTextureBPLibrary::GetBonePositionsAndRotations(const USkeletalMeshC
 	const TArray<FTransform>& CompSpaceTransforms = SkeletalMeshComponent->GetComponentSpaceTransforms();
 	check(CompSpaceTransforms.Num() >= RefToLocals.Num());
 
+	// Allocate
+	BonePositions.SetNumUninitialized(NumBones);
+	BoneRotations.SetNumUninitialized(NumBones);
+
 	for (int32 BoneIndex = 0; BoneIndex < NumBones; BoneIndex++)
 	{
 		// Decompose Transformation (ComponentSpace)
@@ -568,8 +601,8 @@ int32 UAnimToTextureBPLibrary::GetBonePositionsAndRotations(const USkeletalMeshC
 		const FTransform RelativeTransform(RefToLocalMatrix);
 		DecomposeTransformation(RelativeTransform, BoneRelativePosition, BoneRelativeRotation);
 
-		BonePositions.Add(Delta);
-		BoneRotations.Add(BoneRelativeRotation);
+		BonePositions[BoneIndex] = Delta;
+		BoneRotations[BoneIndex] = BoneRelativeRotation;
 	}
 
 	return NumBones;
@@ -577,7 +610,8 @@ int32 UAnimToTextureBPLibrary::GetBonePositionsAndRotations(const USkeletalMeshC
 
 
 void UAnimToTextureBPLibrary::UpdateMaterialInstanceFromDataAsset(UAnimToTextureDataAsset* DataAsset, UMaterialInstanceConstant* MaterialInstance, 
-	const bool bAnimate, const EAnimToTextureNumBoneInfluences NumBoneInfluences, const EMaterialParameterAssociation MaterialParameterAssociation)
+	const bool bAutoPlay, const int32 AnimationIndex, 
+	const EAnimToTextureNumBoneInfluences NumBoneInfluences, const EMaterialParameterAssociation MaterialParameterAssociation)
 {
 	if (!MaterialInstance || !DataAsset)
 	{
@@ -628,7 +662,6 @@ void UAnimToTextureBPLibrary::UpdateMaterialInstanceFromDataAsset(UAnimToTexture
 		
 		VectorParameter = FLinearColor(DataAsset->VertexSizeBBox);
 		UMaterialEditingLibrary::SetMaterialInstanceVectorParameterValue(MaterialInstance, AnimToTextureParamNames::BoundingBoxScale, VectorParameter, MaterialParameterAssociation);
-		UMaterialEditingLibrary::SetMaterialInstanceScalarParameterValue(MaterialInstance, AnimToTextureParamNames::NumFrames, DataAsset->NumFrames, MaterialParameterAssociation);
 		UMaterialEditingLibrary::SetMaterialInstanceScalarParameterValue(MaterialInstance, AnimToTextureParamNames::RowsPerFrame, DataAsset->VertexRowsPerFrame, MaterialParameterAssociation);
 
 		UMaterialEditingLibrary::SetMaterialInstanceTextureParameterValue(MaterialInstance, AnimToTextureParamNames::VertexPositionTexture, DataAsset->GetVertexPositionTexture(), MaterialParameterAssociation);
@@ -645,7 +678,6 @@ void UAnimToTextureBPLibrary::UpdateMaterialInstanceFromDataAsset(UAnimToTexture
 
 		VectorParameter = FLinearColor(DataAsset->BoneSizeBBox);
 		UMaterialEditingLibrary::SetMaterialInstanceVectorParameterValue(MaterialInstance, AnimToTextureParamNames::BoundingBoxScale, VectorParameter, MaterialParameterAssociation);
-		UMaterialEditingLibrary::SetMaterialInstanceScalarParameterValue(MaterialInstance, AnimToTextureParamNames::NumFrames, DataAsset->NumFrames, MaterialParameterAssociation);
 		UMaterialEditingLibrary::SetMaterialInstanceScalarParameterValue(MaterialInstance, AnimToTextureParamNames::RowsPerFrame, DataAsset->BoneRowsPerFrame, MaterialParameterAssociation);
 		UMaterialEditingLibrary::SetMaterialInstanceScalarParameterValue(MaterialInstance, AnimToTextureParamNames::BoneWeightRowsPerFrame, DataAsset->BoneWeightRowsPerFrame, MaterialParameterAssociation);
 
@@ -672,8 +704,19 @@ void UAnimToTextureBPLibrary::UpdateMaterialInstanceFromDataAsset(UAnimToTexture
 
 	}
 
-	// Animate
-	UMaterialEditingLibrary::SetMaterialInstanceStaticSwitchParameterValue(MaterialInstance, AnimToTextureParamNames::Animate, bAnimate, MaterialParameterAssociation);
+	// AutoPlay
+	if (bAutoPlay && DataAsset->Animations.IsValidIndex(AnimationIndex))
+	{
+		UMaterialEditingLibrary::SetMaterialInstanceStaticSwitchParameterValue(MaterialInstance, AnimToTextureParamNames::AutoPlay, true, MaterialParameterAssociation);
+		UMaterialEditingLibrary::SetMaterialInstanceScalarParameterValue(MaterialInstance, AnimToTextureParamNames::StartFrame, DataAsset->Animations[AnimationIndex].StartFrame, MaterialParameterAssociation);
+		UMaterialEditingLibrary::SetMaterialInstanceScalarParameterValue(MaterialInstance, AnimToTextureParamNames::EndFrame, DataAsset->Animations[AnimationIndex].EndFrame, MaterialParameterAssociation);
+	}
+	
+	// NumFrames
+	UMaterialEditingLibrary::SetMaterialInstanceScalarParameterValue(MaterialInstance, AnimToTextureParamNames::NumFrames, DataAsset->NumFrames, MaterialParameterAssociation);
+
+	// SampleRate
+	UMaterialEditingLibrary::SetMaterialInstanceScalarParameterValue(MaterialInstance, AnimToTextureParamNames::SampleRate, DataAsset->SampleRate, MaterialParameterAssociation);
 
 	// Update Material
 	UMaterialEditingLibrary::UpdateMaterialInstance(MaterialInstance);
@@ -703,13 +746,21 @@ bool UAnimToTextureBPLibrary::SetLightMapIndex(UStaticMesh* StaticMesh, const in
 		return false;
 	}
 
-	// 
+	for (int32 Index=0; Index < LightmapIndex; Index++)
+	{
+		if (LightmapIndex > StaticMesh->GetNumUVChannels(LODIndex))
+		{
+			StaticMesh->AddUVChannel(LODIndex);
+		}
+	}
+
+	// Set Build Settings
 	FStaticMeshSourceModel& SourceModel = StaticMesh->GetSourceModel(LODIndex);
 	SourceModel.BuildSettings.bGenerateLightmapUVs = bGenerateLightmapUVs;
 	SourceModel.BuildSettings.DstLightmapIndex = LightmapIndex;
 	StaticMesh->SetLightMapCoordinateIndex(LightmapIndex);
 
-	// Build mesh from source
+	// Build Mesh
 	StaticMesh->Build(false);
 	StaticMesh->PostEditChange();
 	StaticMesh->MarkPackageDirty();
@@ -732,7 +783,7 @@ UStaticMesh* UAnimToTextureBPLibrary::ConvertSkeletalMeshToStaticMesh(USkeletalM
 
 	if (LODIndex >= 0 && !SkeletalMesh->IsValidLODIndex(LODIndex))
 	{
-		UE_LOG(LogTemp, Warning, TEXT("Invalid LODIndex: %i"), LODIndex);
+		UE_LOG(LogAnimToTextureEditor, Warning, TEXT("Invalid LODIndex: %i"), LODIndex);
 		return nullptr;
 	}
 
@@ -800,7 +851,7 @@ UStaticMesh* UAnimToTextureBPLibrary::ConvertSkeletalMeshToStaticMesh(USkeletalM
 		{
 			if (LODIndex >= TempMesh->GetNumSourceModels())
 			{
-				UE_LOG(LogTemp, Warning, TEXT("Invalid Source Model Index: %i"), LODIndex);
+				UE_LOG(LogAnimToTextureEditor, Warning, TEXT("Invalid Source Model Index: %i"), LODIndex);
 				bGeneratedCorrectly = false;
 			}
 			else
@@ -970,13 +1021,13 @@ bool UAnimToTextureBPLibrary::CreateUVChannel(
 	{
 		if (!StaticMesh->InsertUVChannel(LODIndex, UVChannelIndex))
 		{
-			UE_LOG(LogTemp, Warning, TEXT("Unable to Add UVChannel"));
+			UE_LOG(LogAnimToTextureEditor, Warning, TEXT("Unable to Add UVChannel"));
 			return false;
 		}
 	}
 	else if (UVChannelIndex > StaticMesh->GetNumUVChannels(LODIndex))
 	{
-		UE_LOG(LogTemp, Warning, TEXT("UVChannel: %i Out of Range. Number of existing UVChannels: %i"), UVChannelIndex, StaticMesh->GetNumUVChannels(LODIndex));
+		UE_LOG(LogAnimToTextureEditor, Warning, TEXT("UVChannel: %i Out of Range. Number of existing UVChannels: %i"), UVChannelIndex, StaticMesh->GetNumUVChannels(LODIndex));
 		return false;
 	}
 
@@ -1005,7 +1056,7 @@ bool UAnimToTextureBPLibrary::CreateUVChannel(
 	}
 	else
 	{
-		UE_LOG(LogTemp, Warning, TEXT("Unable to Set UVChannel: %i. TexCoords: %i"), UVChannelIndex, TexCoords.Num());
+		UE_LOG(LogAnimToTextureEditor, Warning, TEXT("Unable to Set UVChannel: %i. TexCoords: %i"), UVChannelIndex, TexCoords.Num());
 		return false;
 	};
 
diff --git a/Engine/Plugins/Animation/AnimToTexture/Source/AnimToTextureEditor/Private/AnimToTextureEditor.cpp b/Engine/Plugins/Animation/AnimToTexture/Source/AnimToTextureEditor/Private/AnimToTextureEditor.cpp
deleted file mode 100644
diff --git a/Engine/Plugins/Animation/AnimToTexture/Source/AnimToTextureEditor/Private/AnimToTextureEditorModule.cpp b/Engine/Plugins/Animation/AnimToTexture/Source/AnimToTextureEditor/Private/AnimToTextureEditorModule.cpp
new file mode 100644
--- /dev/null
+++ b/Engine/Plugins/Animation/AnimToTexture/Source/AnimToTextureEditor/Private/AnimToTextureEditorModule.cpp
@@ -0,0 +1,24 @@
+// Copyright Epic Games, Inc. All Rights Reserved.
+
+#include "AnimToTextureEditorModule.h"
+
+#define LOCTEXT_NAMESPACE "FAnimToTextureEditorModule"
+
+DEFINE_LOG_CATEGORY(LogAnimToTextureEditor);
+
+void FAnimToTextureEditorModule::StartupModule()
+{
+	// This code will execute after your module is loaded into memory; the exact timing is specified in the .uplugin file per-module
+	
+}
+
+void FAnimToTextureEditorModule::ShutdownModule()
+{
+	// This function may be called during shutdown to clean up your module.  For modules that support dynamic reloading,
+	// we call this function before unloading the module.
+	
+}
+
+#undef LOCTEXT_NAMESPACE
+	
+IMPLEMENT_MODULE(FAnimToTextureEditorModule, AnimToTextureEditor)
\ No newline at end of file
diff --git a/Engine/Plugins/Animation/AnimToTexture/Source/AnimToTextureEditor/Private/AnimToTextureMeshMapping.cpp b/Engine/Plugins/Animation/AnimToTexture/Source/AnimToTextureEditor/Private/AnimToTextureMeshMapping.cpp
--- a/Engine/Plugins/Animation/AnimToTexture/Source/AnimToTextureEditor/Private/AnimToTextureMeshMapping.cpp
+++ b/Engine/Plugins/Animation/AnimToTexture/Source/AnimToTextureEditor/Private/AnimToTextureMeshMapping.cpp
@@ -140,26 +140,6 @@ int32 FVertexToMeshMapping::FindClosestTriangle(const FVector3f& Point, const TA
 	return ClosestTriangle;
 }
 
-//int32 FindClosestVertex(const FVector3f& Point, const TArray<FVector3f>& Vertices)
-//{
-//	float MinDistance = TNumericLimits<float>::Max();
-//	int32 ClosestIndex = INDEX_NONE;
-//
-//	// Find Closest SkeletalMesh Vertex.
-//	for (int32 VertexIndex = 0; VertexIndex < Vertices.Num(); ++VertexIndex)
-//	{
-//		const float Distance = FVector3f::Dist(Point, Vertices[VertexIndex]);
-//		if (Distance < MinDistance)
-//		{
-//			MinDistance = Distance;
-//			ClosestIndex = VertexIndex;
-//		}
-//	}
-//
-//	return ClosestIndex;
-//}
-
-
 void FVertexToMeshMapping::Create(
 	const UStaticMesh* StaticMesh, const int32 StaticMeshLODIndex,	
 	const USkeletalMesh* SkeletalMesh, const int32 SkeletalMeshLODIndex,
diff --git a/Engine/Plugins/Animation/AnimToTexture/Source/AnimToTextureEditor/Private/AnimToTextureSkeletalMesh.cpp b/Engine/Plugins/Animation/AnimToTexture/Source/AnimToTextureEditor/Private/AnimToTextureSkeletalMesh.cpp
--- a/Engine/Plugins/Animation/AnimToTexture/Source/AnimToTextureEditor/Private/AnimToTextureSkeletalMesh.cpp
+++ b/Engine/Plugins/Animation/AnimToTexture/Source/AnimToTextureEditor/Private/AnimToTextureSkeletalMesh.cpp
@@ -1,5 +1,6 @@
 // Copyright Epic Games, Inc. All Rights Reserved.
 #include "AnimToTextureSkeletalMesh.h"
+#include "AnimToTextureEditorModule.h"
 #include "Rendering/SkeletalMeshRenderData.h"
 #include "Engine/StaticMesh.h"
 #include "Engine/SkeletalMesh.h"
@@ -207,7 +208,7 @@ void GetSkinWeights(const USkeletalMesh* SkeletalMesh, const int32 LODIndex,
 
 	if (!RenderData->LODRenderData.IsValidIndex(LODIndex))
 	{
-		UE_LOG(LogTemp, Warning, TEXT("Invalid LODIndex: %i"), LODIndex)
+		UE_LOG(LogAnimToTextureEditor, Warning, TEXT("Invalid LODIndex: %i"), LODIndex)
 		return;
 	}
 
diff --git a/Engine/Plugins/Animation/AnimToTexture/Source/AnimToTextureEditor/Public/AnimToTextureBPLibrary.h b/Engine/Plugins/Animation/AnimToTexture/Source/AnimToTextureEditor/Public/AnimToTextureBPLibrary.h
--- a/Engine/Plugins/Animation/AnimToTexture/Source/AnimToTextureEditor/Public/AnimToTextureBPLibrary.h
+++ b/Engine/Plugins/Animation/AnimToTexture/Source/AnimToTextureEditor/Public/AnimToTextureBPLibrary.h
@@ -37,16 +37,21 @@ class UAnimToTextureBPLibrary : public UBlueprintFunctionLibrary
 	UFUNCTION(BlueprintCallable, Category = "AnimToTexture")
 	static UStaticMesh* ConvertSkeletalMeshToStaticMesh(USkeletalMesh* SkeletalMesh, const FString PackageName, const int32 LODIndex = -1);
 
-	// FIXME: you cannot set index to 2 if there is no index 1
+	/**
+	* Utility for setting a StaticMesh LightMapIndex.
+	*/
 	UFUNCTION(BlueprintCallable, Category = "AnimToTexture")
 	static bool SetLightMapIndex(UStaticMesh* StaticMesh, const int32 LODIndex, const int32 LightmapIndex=1, bool bGenerateLightmapUVs=true);
 
 	/**
-	* Updates a material's parameters to match those of an animToTexture data asset
+	* Updates a material's parameters to match those of an AnimToTexture DataAsset
+	* @param AnimationIndex Index of the animation to play.
+	* @param NumBoneInfluences Number of skinning influences. More influences will produce smoother deformation but slower performance. This is only used in Bone Mode.
 	*/
 	UFUNCTION(BlueprintCallable, meta = (Category = "AnimToTexture"))
 	static void UpdateMaterialInstanceFromDataAsset(UAnimToTextureDataAsset* DataAsset, class UMaterialInstanceConstant* MaterialInstance, 
-		const bool bAnimate=false, const EAnimToTextureNumBoneInfluences NumBoneInfluences = EAnimToTextureNumBoneInfluences::One,
+		const bool bAutoPlay=true, const int32 AnimationIndex=0, 
+		const EAnimToTextureNumBoneInfluences NumBoneInfluences = EAnimToTextureNumBoneInfluences::One,
 		const EMaterialParameterAssociation MaterialParameterAssociation = EMaterialParameterAssociation::LayerParameter);
 
 #endif // WITH_EDITOR
diff --git a/Engine/Plugins/Animation/AnimToTexture/Source/AnimToTextureEditor/Public/AnimToTextureEditor.h b/Engine/Plugins/Animation/AnimToTexture/Source/AnimToTextureEditor/Public/AnimToTextureEditor.h
deleted file mode 100644
diff --git a/Engine/Plugins/Animation/AnimToTexture/Source/AnimToTextureEditor/Public/AnimToTextureEditorModule.h b/Engine/Plugins/Animation/AnimToTexture/Source/AnimToTextureEditor/Public/AnimToTextureEditorModule.h
new file mode 100644
--- /dev/null
+++ b/Engine/Plugins/Animation/AnimToTexture/Source/AnimToTextureEditor/Public/AnimToTextureEditorModule.h
@@ -0,0 +1,16 @@
+// Copyright Epic Games, Inc. All Rights Reserved.
+
+#pragma once
+
+#include "Modules/ModuleManager.h"
+
+DECLARE_LOG_CATEGORY_EXTERN(LogAnimToTextureEditor, Log, All);
+
+class FAnimToTextureEditorModule : public IModuleInterface
+{
+public:
+
+	/** IModuleInterface implementation */
+	virtual void StartupModule() override;
+	virtual void ShutdownModule() override;
+};
