Moving AnimToTexture files to MassTraffic

diff --git a/Engine/Plugins/Animation/AnimToTexture/Config/BaseAnimToTexture.ini b/Engine/Plugins/Animation/AnimToTexture/Config/BaseAnimToTexture.ini
--- a/Engine/Plugins/Animation/AnimToTexture/Config/BaseAnimToTexture.ini
+++ b/Engine/Plugins/Animation/AnimToTexture/Config/BaseAnimToTexture.ini
@@ -1,3 +1,4 @@
 [CoreRedirects]
 +StructRedirects=(OldName="/Script/AnimToTexture.AnimSequenceInfo",NewName="/Script/AnimToTexture.AnimToTextureAnimSequenceInfo")
 +StructRedirects=(OldName="/Script/AnimToTexture.AnimInfo",NewName="/Script/AnimToTexture.AnimToTextureAnimInfo")
++StructRedirects=(OldName="/Script/AnimToTexture.AnimToTextureInstancePlaybackData",NewName="/Script/AnimToTexture.AnimToTextureFrameData")
