[AnimToTexture]
- Updated plugin version
- Fixing warning: Preview default for UseTwoInfluences and UseFourInfluences

diff --git a/Engine/Plugins/Animation/AnimToTexture/AnimToTexture.uplugin b/Engine/Plugins/Animation/AnimToTexture/AnimToTexture.uplugin
--- a/Engine/Plugins/Animation/AnimToTexture/AnimToTexture.uplugin
+++ b/Engine/Plugins/Animation/AnimToTexture/AnimToTexture.uplugin
@@ -1,7 +1,7 @@
 {
 	"FileVersion": 3,
-	"Version": 1,
-	"VersionName": "1.0",
+	"Version": 2,
+	"VersionName": "2.0",
 	"FriendlyName": "AnimToTexture",
 	"Description": "Converts SkeletalMesh Animations into Textures",
 	"Category": "Other",
