[AnimToTexture] Moving AnimToTexture to Experimental folder

diff --git a/Engine/Plugins/Animation/AnimToTexture/AnimToTexture.uplugin b/Engine/Plugins/Animation/AnimToTexture/AnimToTexture.uplugin
deleted file mode 100644
diff --git a/Engine/Plugins/Animation/AnimToTexture/Config/BaseAnimToTexture.ini b/Engine/Plugins/Animation/AnimToTexture/Config/BaseAnimToTexture.ini
deleted file mode 100644
diff --git a/Engine/Plugins/Animation/AnimToTexture/Config/DefaultAnimToTexture.ini b/Engine/Plugins/Animation/AnimToTexture/Config/DefaultAnimToTexture.ini
deleted file mode 100644
diff --git a/Engine/Plugins/Animation/AnimToTexture/Source/AnimToTexture/AnimToTexture.Build.cs b/Engine/Plugins/Animation/AnimToTexture/Source/AnimToTexture/AnimToTexture.Build.cs
deleted file mode 100644
diff --git a/Engine/Plugins/Animation/AnimToTexture/Source/AnimToTexture/Private/AnimToTextureDataAsset.cpp b/Engine/Plugins/Animation/AnimToTexture/Source/AnimToTexture/Private/AnimToTextureDataAsset.cpp
deleted file mode 100644
diff --git a/Engine/Plugins/Animation/AnimToTexture/Source/AnimToTexture/Private/AnimToTextureInstancePlaybackHelpers.cpp b/Engine/Plugins/Animation/AnimToTexture/Source/AnimToTexture/Private/AnimToTextureInstancePlaybackHelpers.cpp
deleted file mode 100644
diff --git a/Engine/Plugins/Animation/AnimToTexture/Source/AnimToTexture/Private/AnimToTextureModule.cpp b/Engine/Plugins/Animation/AnimToTexture/Source/AnimToTexture/Private/AnimToTextureModule.cpp
deleted file mode 100644
diff --git a/Engine/Plugins/Animation/AnimToTexture/Source/AnimToTexture/Public/AnimToTextureDataAsset.h b/Engine/Plugins/Animation/AnimToTexture/Source/AnimToTexture/Public/AnimToTextureDataAsset.h
deleted file mode 100644
diff --git a/Engine/Plugins/Animation/AnimToTexture/Source/AnimToTexture/Public/AnimToTextureInstancePlaybackHelpers.h b/Engine/Plugins/Animation/AnimToTexture/Source/AnimToTexture/Public/AnimToTextureInstancePlaybackHelpers.h
deleted file mode 100644
diff --git a/Engine/Plugins/Animation/AnimToTexture/Source/AnimToTexture/Public/AnimToTextureModule.h b/Engine/Plugins/Animation/AnimToTexture/Source/AnimToTexture/Public/AnimToTextureModule.h
deleted file mode 100644
diff --git a/Engine/Plugins/Animation/AnimToTexture/Source/AnimToTextureEditor/AnimToTextureEditor.Build.cs b/Engine/Plugins/Animation/AnimToTexture/Source/AnimToTextureEditor/AnimToTextureEditor.Build.cs
deleted file mode 100644
diff --git a/Engine/Plugins/Animation/AnimToTexture/Source/AnimToTextureEditor/Private/AnimToTextureBPLibrary.cpp b/Engine/Plugins/Animation/AnimToTexture/Source/AnimToTextureEditor/Private/AnimToTextureBPLibrary.cpp
deleted file mode 100644
diff --git a/Engine/Plugins/Animation/AnimToTexture/Source/AnimToTextureEditor/Private/AnimToTextureEditorModule.cpp b/Engine/Plugins/Animation/AnimToTexture/Source/AnimToTextureEditor/Private/AnimToTextureEditorModule.cpp
deleted file mode 100644
diff --git a/Engine/Plugins/Animation/AnimToTexture/Source/AnimToTextureEditor/Private/AnimToTextureMeshMapping.cpp b/Engine/Plugins/Animation/AnimToTexture/Source/AnimToTextureEditor/Private/AnimToTextureMeshMapping.cpp
deleted file mode 100644
diff --git a/Engine/Plugins/Animation/AnimToTexture/Source/AnimToTextureEditor/Private/AnimToTextureSkeletalMesh.cpp b/Engine/Plugins/Animation/AnimToTexture/Source/AnimToTextureEditor/Private/AnimToTextureSkeletalMesh.cpp
deleted file mode 100644
diff --git a/Engine/Plugins/Animation/AnimToTexture/Source/AnimToTextureEditor/Private/AnimToTextureUtils.cpp b/Engine/Plugins/Animation/AnimToTexture/Source/AnimToTextureEditor/Private/AnimToTextureUtils.cpp
deleted file mode 100644
diff --git a/Engine/Plugins/Animation/AnimToTexture/Source/AnimToTextureEditor/Public/AnimToTextureBPLibrary.h b/Engine/Plugins/Animation/AnimToTexture/Source/AnimToTextureEditor/Public/AnimToTextureBPLibrary.h
deleted file mode 100644
diff --git a/Engine/Plugins/Animation/AnimToTexture/Source/AnimToTextureEditor/Public/AnimToTextureEditorModule.h b/Engine/Plugins/Animation/AnimToTexture/Source/AnimToTextureEditor/Public/AnimToTextureEditorModule.h
deleted file mode 100644
diff --git a/Engine/Plugins/Animation/AnimToTexture/Source/AnimToTextureEditor/Public/AnimToTextureMeshMapping.h b/Engine/Plugins/Animation/AnimToTexture/Source/AnimToTextureEditor/Public/AnimToTextureMeshMapping.h
deleted file mode 100644
diff --git a/Engine/Plugins/Animation/AnimToTexture/Source/AnimToTextureEditor/Public/AnimToTextureSkeletalMesh.h b/Engine/Plugins/Animation/AnimToTexture/Source/AnimToTextureEditor/Public/AnimToTextureSkeletalMesh.h
deleted file mode 100644
diff --git a/Engine/Plugins/Animation/AnimToTexture/Source/AnimToTextureEditor/Public/AnimToTextureUtils.h b/Engine/Plugins/Animation/AnimToTexture/Source/AnimToTextureEditor/Public/AnimToTextureUtils.h
deleted file mode 100644
diff --git a/Engine/Plugins/Experimental/AnimToTexture/AnimToTexture.uplugin b/Engine/Plugins/Experimental/AnimToTexture/AnimToTexture.uplugin
new file mode 100644
--- /dev/null
+++ b/Engine/Plugins/Experimental/AnimToTexture/AnimToTexture.uplugin
@@ -0,0 +1,29 @@
+{
+	"FileVersion": 3,
+	"Version": 2,
+	"VersionName": "2.0",
+	"FriendlyName": "AnimToTexture",
+	"Description": "Converts SkeletalMesh Animations into Textures",
+	"Category": "Other",
+	"CreatedBy": "Epic Games, Inc.",
+	"CreatedByURL": "https://epicgames.com",
+	"DocsURL": "",
+	"MarketplaceURL": "",
+	"SupportURL": "",
+	"CanContainContent": true,
+	"IsBetaVersion": false,
+	"IsExperimentalVersion": true,
+	"Installed": false,
+	"Modules": [
+		{
+			"Name": "AnimToTexture",
+			"Type": "Runtime",
+			"LoadingPhase": "PreDefault"
+		},
+		{
+			"Name": "AnimToTextureEditor",
+			"Type": "Editor",
+			"LoadingPhase": "PreLoadingScreen"
+		}
+	]
+}
\ No newline at end of file
diff --git a/Engine/Plugins/Experimental/AnimToTexture/Config/BaseAnimToTexture.ini b/Engine/Plugins/Experimental/AnimToTexture/Config/BaseAnimToTexture.ini
new file mode 100644
--- /dev/null
+++ b/Engine/Plugins/Experimental/AnimToTexture/Config/BaseAnimToTexture.ini
@@ -0,0 +1,3 @@
+[CoreRedirects]
++StructRedirects=(OldName="/Script/AnimToTexture.AnimSequenceInfo",NewName="/Script/AnimToTexture.AnimToTextureAnimSequenceInfo")
++StructRedirects=(OldName="/Script/AnimToTexture.AnimInfo",NewName="/Script/AnimToTexture.AnimToTextureAnimInfo")
\ No newline at end of file
diff --git a/Engine/Plugins/Experimental/AnimToTexture/Config/DefaultAnimToTexture.ini b/Engine/Plugins/Experimental/AnimToTexture/Config/DefaultAnimToTexture.ini
new file mode 100644
--- /dev/null
+++ b/Engine/Plugins/Experimental/AnimToTexture/Config/DefaultAnimToTexture.ini
@@ -0,0 +1,3 @@
+[CoreRedirects]
++PropertyRedirects==(OldName="AnimToTextureDataAsset.MasterSkeletalMesh",NewName="LeaderSkeletalMesh")
++FunctionRedirects=(OldName="AnimToTextureDataAsset.BP_GetMasterSkeletalMesh", NewName="BP_GetLeaderSkeletalMesh")
\ No newline at end of file
diff --git a/Engine/Plugins/Experimental/AnimToTexture/Source/AnimToTexture/AnimToTexture.Build.cs b/Engine/Plugins/Experimental/AnimToTexture/Source/AnimToTexture/AnimToTexture.Build.cs
new file mode 100644
--- /dev/null
+++ b/Engine/Plugins/Experimental/AnimToTexture/Source/AnimToTexture/AnimToTexture.Build.cs
@@ -0,0 +1,54 @@
+// Copyright Epic Games, Inc. All Rights Reserved.
+
+using UnrealBuildTool;
+
+public class AnimToTexture : ModuleRules
+{
+	public AnimToTexture(ReadOnlyTargetRules Target) : base(Target)
+	{
+		PCHUsage = ModuleRules.PCHUsageMode.UseExplicitOrSharedPCHs;
+		
+		PublicIncludePaths.AddRange(
+			new string[] {
+				// ... add public include paths required here ...
+			}
+			);
+				
+		
+		PrivateIncludePaths.AddRange(
+			new string[] {
+				// ... add other private include paths required here ...
+			}
+			);
+			
+		
+		PublicDependencyModuleNames.AddRange(
+			new string[]
+			{
+				"Core",
+				// ... add other public dependencies that you statically link with here ...
+			}
+			);
+			
+		
+		PrivateDependencyModuleNames.AddRange(
+			new string[]
+			{
+				"CoreUObject",
+				"Engine",
+				"Slate",
+				"SlateCore",
+				// ... add private dependencies that you statically link with here ...	
+				"MeshDescription",
+			}
+			);
+		
+		
+		DynamicallyLoadedModuleNames.AddRange(
+			new string[]
+			{
+				// ... add any modules that your module loads dynamically here ...
+			}
+			);
+	}
+}
diff --git a/Engine/Plugins/Experimental/AnimToTexture/Source/AnimToTexture/Private/AnimToTextureDataAsset.cpp b/Engine/Plugins/Experimental/AnimToTexture/Source/AnimToTexture/Private/AnimToTextureDataAsset.cpp
new file mode 100644
--- /dev/null
+++ b/Engine/Plugins/Experimental/AnimToTexture/Source/AnimToTexture/Private/AnimToTextureDataAsset.cpp
@@ -0,0 +1,38 @@
+// Copyright Epic Games, Inc. All Rights Reserved.
+
+#include "AnimToTextureDataAsset.h"
+
+int32 UAnimToTextureDataAsset::GetIndexFromAnimSequence(const UAnimSequence* Sequence)
+{
+	const int32 NumSequences = AnimSequences.Num();
+	
+	// We can store a sequence to index map for a faster search
+	for (int32 Index = 0; Index < NumSequences; ++Index)
+	{
+		const FAnimToTextureAnimSequenceInfo& SequenceInfo = AnimSequences[Index];
+		if (SequenceInfo.bEnabled && SequenceInfo.AnimSequence == Sequence)
+		{
+			return Index;
+		}
+	}
+
+	return INDEX_NONE;
+}
+
+void UAnimToTextureDataAsset::ResetInfo()
+{
+	// Common Info.
+	NumFrames = 0;
+	Animations.Reset();
+
+	// Vertex Info
+	VertexRowsPerFrame = 1;
+	VertexMinBBox = FVector::ZeroVector;
+	VertexSizeBBox = FVector::ZeroVector;
+
+	// Bone Info
+	BoneRowsPerFrame = 1;
+	BoneWeightRowsPerFrame = 1;
+	BoneMinBBox = FVector::ZeroVector;
+	BoneSizeBBox = FVector::ZeroVector;
+};
\ No newline at end of file
diff --git a/Engine/Plugins/Experimental/AnimToTexture/Source/AnimToTexture/Private/AnimToTextureInstancePlaybackHelpers.cpp b/Engine/Plugins/Experimental/AnimToTexture/Source/AnimToTexture/Private/AnimToTextureInstancePlaybackHelpers.cpp
new file mode 100644
--- /dev/null
+++ b/Engine/Plugins/Experimental/AnimToTexture/Source/AnimToTexture/Private/AnimToTextureInstancePlaybackHelpers.cpp
@@ -0,0 +1,112 @@
+// Copyright Epic Games, Inc. All Rights Reserved.
+
+#include "AnimToTextureInstancePlaybackHelpers.h"
+#include "AnimToTextureDataAsset.h"
+
+bool UAnimToTextureInstancePlaybackLibrary::SetupInstancedMeshComponent(UInstancedStaticMeshComponent* InstancedMeshComponent, int32 NumInstances, bool bAutoPlay)
+{	
+	if (!InstancedMeshComponent)
+	{
+		return false;
+	}
+
+	// Clear data.
+	InstancedMeshComponent->ClearInstances();
+
+	if (!NumInstances)
+	{
+		return false;
+	}
+
+	// Allocate Transforms
+	TArray<FTransform> InstanceTransforms;
+	InstanceTransforms.AddDefaulted(NumInstances);
+
+	// Set Custom Data Length
+	const SIZE_T DataSize = bAutoPlay ? sizeof(FAnimToTextureAutoPlayData) : sizeof(FAnimToTextureFrameData);
+	InstancedMeshComponent->NumCustomDataFloats = DataSize / sizeof(float);
+
+	// Initizalize Instances
+	InstancedMeshComponent->AddInstances(InstanceTransforms, /*bShouldReturnIndices*/ false, /*bWorldSpace*/ false);
+
+	return true;
+}
+
+bool UAnimToTextureInstancePlaybackLibrary::BatchUpdateInstancesAutoPlayData(UInstancedStaticMeshComponent* InstancedMeshComponent, 
+	const TArray<FAnimToTextureAutoPlayData>& AutoPlayData, const TArray<FMatrix>& Transforms, bool bMarkRenderStateDirty)
+{
+	return BatchUpdateInstancesData< FAnimToTextureAutoPlayData>(InstancedMeshComponent, AutoPlayData, Transforms, bMarkRenderStateDirty);
+}
+
+bool UAnimToTextureInstancePlaybackLibrary::BatchUpdateInstancesFrameData(UInstancedStaticMeshComponent* InstancedMeshComponent,
+	const TArray<FAnimToTextureFrameData>& FrameData, const TArray<FMatrix>& Transforms, bool bMarkRenderStateDirty)
+{
+	return BatchUpdateInstancesData<FAnimToTextureFrameData>(InstancedMeshComponent, FrameData, Transforms, bMarkRenderStateDirty);
+}
+
+bool UAnimToTextureInstancePlaybackLibrary::UpdateInstanceAutoPlayData(UInstancedStaticMeshComponent* InstancedMeshComponent,
+	int32 InstanceIndex, const FAnimToTextureAutoPlayData& AutoPlayData, bool bMarkRenderStateDirty)
+{
+	return UpdateInstanceCustomData<FAnimToTextureAutoPlayData>(InstancedMeshComponent, InstanceIndex, AutoPlayData, bMarkRenderStateDirty);
+}
+
+bool UAnimToTextureInstancePlaybackLibrary::UpdateInstanceFrameData(UInstancedStaticMeshComponent* InstancedMeshComponent,
+	int32 InstanceIndex, const FAnimToTextureFrameData& FrameData, bool bMarkRenderStateDirty)
+{
+	return UpdateInstanceCustomData<FAnimToTextureFrameData>(InstancedMeshComponent, InstanceIndex, FrameData, bMarkRenderStateDirty);
+}
+
+bool UAnimToTextureInstancePlaybackLibrary::GetAutoPlayDataFromDataAsset(const UAnimToTextureDataAsset* DataAsset, int32 AnimationIndex,
+	FAnimToTextureAutoPlayData& AutoPlayData, float TimeOffset, float PlayRate)
+{
+	if (DataAsset && DataAsset->Animations.IsValidIndex(AnimationIndex))
+	{
+		// Copy Frame Range
+		const FAnimToTextureAnimInfo& AnimInfo = DataAsset->Animations[AnimationIndex];
+		AutoPlayData.StartFrame = AnimInfo.StartFrame;
+		AutoPlayData.EndFrame = AnimInfo.EndFrame;
+
+		AutoPlayData.TimeOffset = TimeOffset;
+		AutoPlayData.PlayRate = PlayRate;
+
+		return true;
+	}
+
+	// Return default
+	AutoPlayData = FAnimToTextureAutoPlayData();
+	return false;
+}
+
+float UAnimToTextureInstancePlaybackLibrary::GetFrame(float Time, float StartFrame, float EndFrame, 
+	float TimeOffset, float PlayRate, float SampleRate)
+{	
+	// Clamp inputs (just in case)
+	Time = FMath::Max(Time, 0.f);
+	TimeOffset = FMath::Max(TimeOffset, 0.f);
+	PlayRate = FMath::Max(PlayRate, 0.f);
+
+	const float Frame = (Time + TimeOffset) * (PlayRate * SampleRate);
+	const float NumFrames = EndFrame - StartFrame + 1.f;
+	return FMath::Fmod(Frame, NumFrames) + StartFrame;
+}
+
+bool UAnimToTextureInstancePlaybackLibrary::GetFrameDataFromDataAsset(const UAnimToTextureDataAsset* DataAsset, int32 AnimationIndex, float Time, FAnimToTextureFrameData& FrameData, float TimeOffset, float PlayRate)
+{
+	if (DataAsset && DataAsset->Animations.IsValidIndex(AnimationIndex))
+	{
+		// Copy Frame Range
+		const FAnimToTextureAnimInfo& AnimInfo = DataAsset->Animations[AnimationIndex];
+		
+		FrameData.Frame = GetFrame(Time, AnimInfo.StartFrame, AnimInfo.EndFrame, 
+			TimeOffset, PlayRate, DataAsset->SampleRate);
+
+		// Previous Frame
+		FrameData.PrevFrame = FMath::Clamp(FrameData.Frame - 1, AnimInfo.StartFrame, AnimInfo.EndFrame);
+
+		return true;
+	}
+
+	// Return default
+	FrameData = FAnimToTextureFrameData();
+	return false;
+}
diff --git a/Engine/Plugins/Experimental/AnimToTexture/Source/AnimToTexture/Private/AnimToTextureModule.cpp b/Engine/Plugins/Experimental/AnimToTexture/Source/AnimToTexture/Private/AnimToTextureModule.cpp
new file mode 100644
--- /dev/null
+++ b/Engine/Plugins/Experimental/AnimToTexture/Source/AnimToTexture/Private/AnimToTextureModule.cpp
@@ -0,0 +1,24 @@
+// Copyright Epic Games, Inc. All Rights Reserved.
+
+#include "AnimToTextureModule.h"
+
+DEFINE_LOG_CATEGORY(LogAnimToTexture);
+
+#define LOCTEXT_NAMESPACE "FAnimToTextureModule"
+
+void FAnimToTextureModule::StartupModule()
+{
+	// This code will execute after your module is loaded into memory; the exact timing is specified in the .uplugin file per-module
+	
+}
+
+void FAnimToTextureModule::ShutdownModule()
+{
+	// This function may be called during shutdown to clean up your module.  For modules that support dynamic reloading,
+	// we call this function before unloading the module.
+	
+}
+
+#undef LOCTEXT_NAMESPACE
+	
+IMPLEMENT_MODULE(FAnimToTextureModule, AnimToTexture)
\ No newline at end of file
diff --git a/Engine/Plugins/Experimental/AnimToTexture/Source/AnimToTexture/Public/AnimToTextureDataAsset.h b/Engine/Plugins/Experimental/AnimToTexture/Source/AnimToTexture/Public/AnimToTextureDataAsset.h
new file mode 100644
--- /dev/null
+++ b/Engine/Plugins/Experimental/AnimToTexture/Source/AnimToTexture/Public/AnimToTextureDataAsset.h
@@ -0,0 +1,317 @@
+// Copyright Epic Games, Inc. All Rights Reserved.
+
+#pragma once
+
+#include "Engine/DataAsset.h"
+#include "Animation/AnimInstance.h"
+#include "Engine/SkeletalMesh.h"
+#include "Engine/StaticMesh.h"
+#include "Engine/Texture2D.h"
+#include "AnimToTextureDataAsset.generated.h"
+
+class USkeletalMesh;
+class UStaticMesh;
+class UTexture2D;
+
+namespace AnimToTextureParamNames
+{
+	static const FName Frame = TEXT("Frame");
+	static const FName AutoPlay = TEXT("AutoPlay");
+	static const FName StartFrame = TEXT("StartFrame");
+	static const FName EndFrame = TEXT("EndFrame");
+	static const FName SampleRate = TEXT("SampleRate");
+	static const FName NumFrames = TEXT("NumFrames");
+	static const FName BoundingBoxMin = TEXT("MinBBox");
+	static const FName BoundingBoxScale = TEXT("SizeBBox");
+	static const FName RowsPerFrame = TEXT("RowsPerFrame");
+	static const FName BoneWeightRowsPerFrame = TEXT("BoneWeightsRowsPerFrame");
+	static const FName VertexPositionTexture = TEXT("PositionTexture");
+	static const FName VertexNormalTexture = TEXT("NormalTexture");
+	static const FName BonePositionTexture = TEXT("BonePositionTexture");
+	static const FName BoneRotationTexture = TEXT("BoneRotationTexture");
+	static const FName BoneWeightsTexture = TEXT("BoneWeightsTexture");
+	static const FName UseUV0 = TEXT("UseUV0");
+	static const FName UseUV1 = TEXT("UseUV1");
+	static const FName UseUV2 = TEXT("UseUV2");
+	static const FName UseUV3 = TEXT("UseUV3");
+	static const FName UseTwoInfluences = TEXT("UseTwoInfluences");
+	static const FName UseFourInfluences = TEXT("UseFourInfluences");
+}
+
+UENUM(Blueprintable)
+enum class EAnimToTextureMode : uint8
+{
+	/* Position and Normal Per-Vertex */
+	Vertex,
+	/* Linear Blending Skinnin */
+	Bone,
+};
+
+UENUM(Blueprintable)
+enum class EAnimToTexturePrecision : uint8
+{
+	/* 8 bits */
+	EightBits,
+	/* 16 bits */
+	SixteenBits,
+};
+
+UENUM(Blueprintable)
+enum class EAnimToTextureNumBoneInfluences : uint8
+{
+	/* Single bone influence */
+	One,
+	/* Blend between two influences */
+	Two,
+	/* Blend between four influences */
+	Four,
+};
+
+USTRUCT(Blueprintable)
+struct FAnimToTextureAnimSequenceInfo
+{
+	GENERATED_BODY()
+
+	UPROPERTY(EditAnywhere, Category = Default, BlueprintReadWrite)
+	bool bEnabled = true;
+
+	UPROPERTY(EditAnywhere, Category = Default, BlueprintReadWrite, meta = (EditCondition = "bEnabled", EditConditionHides))
+	TObjectPtr<UAnimSequence> AnimSequence = nullptr;
+
+	UPROPERTY(EditAnywhere, Category = Default, BlueprintReadWrite, meta = (EditCondition = "bEnabled", EditConditionHides))
+	bool bUseCustomRange = false;
+
+	UPROPERTY(EditAnywhere, Category = Default, BlueprintReadWrite, meta = (EditCondition = "bEnabled && bUseCustomRange", EditConditionHides))
+	int32 StartFrame = 0;
+
+	UPROPERTY(EditAnywhere, Category = Default, BlueprintReadWrite, meta = (EditCondition = "bEnabled && bUseCustomRange", EditConditionHides))
+	int32 EndFrame = 1;
+
+};
+
+USTRUCT(Blueprintable)
+struct FAnimToTextureAnimInfo
+{
+	GENERATED_BODY()
+	
+	UPROPERTY(VisibleAnywhere, Category = Default, BlueprintReadOnly)
+	int32 StartFrame = 0;
+
+	UPROPERTY(VisibleAnywhere, Category = Default, BlueprintReadOnly)
+	int32 EndFrame = 0;
+
+};
+
+UCLASS(Blueprintable, BlueprintType)
+class ANIMTOTEXTURE_API UAnimToTextureDataAsset : public UPrimaryDataAsset
+{
+public:
+	GENERATED_BODY()
+
+	/**
+	* SkeletalMesh to bake animations from.
+	*/
+	UPROPERTY(EditAnywhere, BlueprintReadWrite, Category = "SkeletalMesh", meta = (AssetBundles = "Client"))
+	TSoftObjectPtr<USkeletalMesh> SkeletalMesh;
+
+	/**
+	* SkeletalMesh LOD to bake.
+	*/
+	UPROPERTY(EditAnywhere, BlueprintReadWrite, Category = "SkeletalMesh", Meta = (DisplayName = "LODIndex"))
+	int32 SkeletalLODIndex = 0;
+
+	/**
+	* StaticMesh to bake to.
+	*/
+	UPROPERTY(EditAnywhere, BlueprintReadWrite, Category = "StaticMesh", meta = (AssetBundles = "Client"))
+	TSoftObjectPtr<UStaticMesh> StaticMesh;
+
+	/**
+	* StaticMesh LOD to bake to.
+	*/
+	UPROPERTY(EditAnywhere, BlueprintReadWrite, Category = "StaticMesh", Meta = (DisplayName = "LODIndex"))
+	int32 StaticLODIndex = 0;
+
+	/**
+	* StaticMesh UVChannel Index for storing vertex information.
+	* Make sure this index does not conflict with the Lightmap UV Index.
+	*/
+	UPROPERTY(EditAnywhere, BlueprintReadWrite, Category = "StaticMesh")
+	int32 UVChannel = 1;
+
+	// ------------------------------------------------------
+	// Texture
+
+	/**
+	* Max resolution of the texture.
+	* A smaller size will be used if the data fits.
+	*/
+	UPROPERTY(EditAnywhere, BlueprintReadWrite, Category = "Texture")
+	int32 MaxHeight = 4096;
+
+	/**
+	* Max resolution of the texture.
+	* A smaller size will be used if the data fits.
+	*/
+	UPROPERTY(EditAnywhere, BlueprintReadWrite, Category = "Texture")
+	int32 MaxWidth = 4096;
+
+	/**
+	* Enforce Power Of Two on texture resolutions.
+	*/
+	UPROPERTY(EditAnywhere, BlueprintReadWrite, Category = "Texture")
+	bool bEnforcePowerOfTwo = false;
+
+	/**
+	* Texture Precision
+	*/
+	UPROPERTY(EditAnywhere, BlueprintReadWrite, Category = "Texture")
+	EAnimToTexturePrecision Precision = EAnimToTexturePrecision::EightBits;
+
+	/**
+	* Storage Mode.
+	* Vertex: will store per-vertex position and normal.
+	* Bone: Will store per-bone position and rotation and per-vertex bone weight. 
+	        This is the preferred method if meshes share skeleton.
+	*/
+	UPROPERTY(EditAnywhere, BlueprintReadWrite, Category = "Texture")
+	EAnimToTextureMode Mode;
+
+	/**
+	* Texture for storing vertex positions
+	* This is only used on Vertex Mode
+	*/
+	UPROPERTY(EditAnywhere, BlueprintReadWrite, Category = "Texture", meta = (EditCondition = "Mode == EAnimToTextureMode::Vertex", EditConditionHides))
+	TSoftObjectPtr<UTexture2D> VertexPositionTexture;
+
+	/**
+	* Texture for storing vertex normals
+	* This is only used on Vertex Mode
+	*/
+	UPROPERTY(EditAnywhere, BlueprintReadWrite, Category = "Texture", meta = (EditCondition = "Mode == EAnimToTextureMode::Vertex", EditConditionHides))
+	TSoftObjectPtr<UTexture2D> VertexNormalTexture;
+
+	/**
+	* Texture for storing bone positions
+	* This is only used on Bone Mode
+	*/
+	UPROPERTY(EditAnywhere, BlueprintReadWrite, Category = "Texture", meta = (EditCondition = "Mode == EAnimToTextureMode::Bone", EditConditionHides))
+	TSoftObjectPtr<UTexture2D> BonePositionTexture;
+
+	/**
+	* Texture for storing bone rotations
+	* This is only used on Bone Mode
+	*/
+	UPROPERTY(EditAnywhere, BlueprintReadWrite, Category = "Texture", meta = (EditCondition = "Mode == EAnimToTextureMode::Bone", EditConditionHides))
+	TSoftObjectPtr<UTexture2D> BoneRotationTexture;
+
+	/**
+	* Texture for storing vertex/bone weighting
+	* This is only used on Bone Mode
+	*/
+	UPROPERTY(EditAnywhere, BlueprintReadWrite, Category = "Texture", meta = (EditCondition = "Mode == EAnimToTextureMode::Bone", EditConditionHides))
+	TSoftObjectPtr<UTexture2D> BoneWeightTexture;
+
+	// ------------------------------------------------------
+	// Animation
+
+	/** 
+	* Bone used for Rigid Binding. The bone needs to be part of the RawBones. 
+	* Sockets and VirtualBones are not supported.
+	*/
+	UPROPERTY(EditAnywhere, BlueprintReadWrite, Category = "Animation", meta = (EditCondition = "Mode == EAnimToTextureMode::Bone"))
+	FName AttachToSocket;
+
+	UPROPERTY(EditAnywhere, BlueprintReadWrite, Category = "Animation")
+	float SampleRate = 30.0f;
+
+	UPROPERTY(EditAnywhere, BlueprintReadWrite, Category = "Animation")
+	TArray<FAnimToTextureAnimSequenceInfo> AnimSequences;
+
+	// ------------------------------------------------------
+	// Info
+
+	/* Total Number of Frames in all animations */
+	UPROPERTY(VisibleAnywhere, BlueprintReadOnly, Category = "Info")
+	int32 NumFrames = 0;
+
+	UPROPERTY(VisibleAnywhere, BlueprintReadOnly, Category = "Info", Meta = (EditCondition = "Mode == EAnimToTextureMode::Vertex", EditConditionHides))
+	int32 VertexRowsPerFrame = 1;
+
+	UPROPERTY(VisibleAnywhere, BlueprintReadOnly, Category = "Info", Meta = (DisplayName = "MinBBox", EditCondition = "Mode == EAnimToTextureMode::Vertex", EditConditionHides))
+	FVector VertexMinBBox;
+
+	UPROPERTY(VisibleAnywhere, BlueprintReadOnly, Category = "Info", Meta = (DisplayName = "SizeBBox", EditCondition = "Mode == EAnimToTextureMode::Vertex", EditConditionHides))
+	FVector VertexSizeBBox;
+
+	UPROPERTY(VisibleAnywhere, BlueprintReadOnly, Category = "Info", Meta = (EditCondition = "Mode == EAnimToTextureMode::Bone", EditConditionHides))
+	int32 BoneWeightRowsPerFrame = 1;
+
+	UPROPERTY(VisibleAnywhere, BlueprintReadOnly, Category = "Info", Meta = (EditCondition = "Mode == EAnimToTextureMode::Bone", EditConditionHides))
+	int32 BoneRowsPerFrame = 1;
+
+	UPROPERTY(VisibleAnywhere, BlueprintReadOnly, Category = "Info", Meta = (DisplayName = "MinBBox", EditCondition = "Mode == EAnimToTextureMode::Bone", EditConditionHides))
+	FVector BoneMinBBox;
+
+	UPROPERTY(VisibleAnywhere, BlueprintReadOnly, Category = "Info", Meta = (DisplayName = "SizeBBox", EditCondition = "Mode == EAnimToTextureMode::Bone", EditConditionHides))
+	FVector BoneSizeBBox;
+
+	UPROPERTY(VisibleAnywhere, BlueprintReadOnly, Category = "Info")
+	TArray<FAnimToTextureAnimInfo> Animations;
+
+	/* Finds AnimSequence Index in the Animations Array. 
+	*  Only Enabled elements are returned.
+	*  Returns -1 if not found.
+	*/
+	UFUNCTION(BlueprintCallable, Category = Default)
+	int32 GetIndexFromAnimSequence(const UAnimSequence* Sequence);
+
+	UFUNCTION()
+	void ResetInfo();
+
+	// If we weren't in a plugin, we could unify this in a base class
+	template<typename AssetType>
+	static AssetType* GetAsset(const TSoftObjectPtr<AssetType>& AssetPointer)
+	{
+		AssetType* ReturnVal = nullptr;
+		if (AssetPointer.ToSoftObjectPath().IsValid())
+		{
+			ReturnVal = AssetPointer.Get();
+			if (!ReturnVal)
+			{
+				AssetType* LoadedAsset = Cast<AssetType>(AssetPointer.ToSoftObjectPath().TryLoad());
+				if (ensureMsgf(LoadedAsset, TEXT("Failed to load asset pointer %s"), *AssetPointer.ToString()))
+				{
+					ReturnVal = LoadedAsset;
+				}
+			}
+		}
+		return ReturnVal;
+	}
+
+#define AnimToTextureDataAsset_ASSET_ACCESSOR(ClassName, PropertyName) \
+	FORCEINLINE ClassName* Get##PropertyName() const { return GetAsset(PropertyName); }
+
+	AnimToTextureDataAsset_ASSET_ACCESSOR(UStaticMesh, StaticMesh);
+	AnimToTextureDataAsset_ASSET_ACCESSOR(USkeletalMesh, SkeletalMesh);
+	AnimToTextureDataAsset_ASSET_ACCESSOR(UTexture2D, VertexPositionTexture);
+	AnimToTextureDataAsset_ASSET_ACCESSOR(UTexture2D, VertexNormalTexture);
+	AnimToTextureDataAsset_ASSET_ACCESSOR(UTexture2D, BonePositionTexture);
+	AnimToTextureDataAsset_ASSET_ACCESSOR(UTexture2D, BoneRotationTexture);
+	AnimToTextureDataAsset_ASSET_ACCESSOR(UTexture2D, BoneWeightTexture);
+
+	UFUNCTION(BlueprintPure, Category = Default, meta = (DisplayName = "Get Static Mesh"))
+	UStaticMesh* BP_GetStaticMesh() { return GetStaticMesh(); }
+
+	UFUNCTION(BlueprintPure, Category = Default, meta = (DisplayName = "Get Skeletal Mesh"))
+	USkeletalMesh* BP_GetSkeletalMesh() { return GetSkeletalMesh(); }
+
+	UFUNCTION(BlueprintPure, Category = Default, meta = (DisplayName = "Get Bone Position Texture"))
+	UTexture2D* BP_GetBonePositionTexture() { return GetBonePositionTexture(); }
+
+	UFUNCTION(BlueprintPure, Category = Default, meta = (DisplayName = "Get Bone Rotation Texture"))
+	UTexture2D* BP_GetBoneRotationTexture() { return GetBoneRotationTexture(); }
+
+	UFUNCTION(BlueprintPure, Category = Default, meta = (DisplayName = "Get Bone Weight Texture"))
+	UTexture2D* BP_GetBoneWeightTexture() { return GetBoneWeightTexture(); }
+};
diff --git a/Engine/Plugins/Experimental/AnimToTexture/Source/AnimToTexture/Public/AnimToTextureInstancePlaybackHelpers.h b/Engine/Plugins/Experimental/AnimToTexture/Source/AnimToTexture/Public/AnimToTextureInstancePlaybackHelpers.h
new file mode 100644
--- /dev/null
+++ b/Engine/Plugins/Experimental/AnimToTexture/Source/AnimToTexture/Public/AnimToTextureInstancePlaybackHelpers.h
@@ -0,0 +1,213 @@
+// Copyright Epic Games, Inc. All Rights Reserved.
+
+#pragma once
+
+#include "Kismet/BlueprintFunctionLibrary.h"
+#include "Components/InstancedStaticMeshComponent.h"
+
+#include "AnimToTextureInstancePlaybackHelpers.generated.h"
+
+class UAnimToTextureDataAsset;
+
+USTRUCT(BlueprintType)
+struct FAnimToTextureFrameData
+{
+	GENERATED_USTRUCT_BODY()
+
+	/**
+	* Frame to be played
+	*/
+	UPROPERTY(EditAnywhere, BlueprintReadWrite, Category = "AnimToTexture|Playback")
+	float Frame = 0.0f;
+
+	/**
+	* Previous Frame (this is needeed for motion blur)
+	*/
+	UPROPERTY(EditAnywhere, BlueprintReadWrite, Category = "AnimToTexture|Playback")
+	float PrevFrame = 0.0f;
+};
+
+USTRUCT(BlueprintType)
+struct FAnimToTextureAutoPlayData
+{
+	GENERATED_USTRUCT_BODY()
+
+	/**
+	* Adds offset to time
+	*/
+	UPROPERTY(EditAnywhere, BlueprintReadWrite, Category = "AnimToTexture|Playback")
+	float TimeOffset = 0.0f;
+
+	/**
+	* Rate for increasing and decreasing speed.
+	*/
+	UPROPERTY(EditAnywhere, BlueprintReadWrite, Category = "AnimToTexture|Playback")
+	float PlayRate = 1.0f;
+
+	/**
+	* Starting frame for animation. 
+	*/
+	UPROPERTY(EditAnywhere, BlueprintReadWrite, Category = "AnimToTexture|Playback")
+	float StartFrame = 0.0f;
+
+	/**
+	* Last frame of animation
+	*/
+	UPROPERTY(EditAnywhere, BlueprintReadWrite, Category = "AnimToTexture|Playback")
+	float EndFrame = 1.0f;
+};
+
+
+UCLASS()
+class ANIMTOTEXTURE_API UAnimToTextureInstancePlaybackLibrary : public UBlueprintFunctionLibrary
+{
+	GENERATED_BODY()
+
+public:
+	
+	/**
+	* Adds Instances and allocates the necessary CustomData.
+	* @param bAutoPlay: if true, it will allocate the required CustomData for working with AutoPlayData. If false, FrameData will be allocated instead./
+	*/
+	UFUNCTION(BlueprintCallable, Category = "AnimToTexture|Playback")
+	static bool SetupInstancedMeshComponent(UInstancedStaticMeshComponent* InstancedMeshComponent, int32 NumInstances, bool bAutoPlay);
+
+	/**
+	* Updates all instances with the given Transforms and AutoPlayData
+	* @param bMarkRenderStateDirty: if true, the change should be visible immediatelly. If you are updating many instances you should only set this to true for the last instance
+	*/
+	UFUNCTION(BlueprintCallable, Category = "AnimToTexture|Playback")
+	static bool BatchUpdateInstancesAutoPlayData(UInstancedStaticMeshComponent* InstancedMeshComponent,
+		const TArray<FAnimToTextureAutoPlayData>& AutoPlayData, const TArray<FMatrix>& Transforms, bool bMarkRenderStateDirty=true);
+
+	/**
+	* Updates all instances with the given Transforms and FrameData
+	* @param bMarkRenderStateDirty: if true, the change should be visible immediatelly. If you are updating many instances you should only set this to true for the last instance
+	*/
+	UFUNCTION(BlueprintCallable, Category = "AnimToTexture|Playback")
+	static bool BatchUpdateInstancesFrameData(UInstancedStaticMeshComponent* InstancedMeshComponent,
+		const TArray<FAnimToTextureFrameData>& FrameData, const TArray<FMatrix>& Transforms, bool bMarkRenderStateDirty=true);
+
+	/**
+	* Updates a single instance with given AutoPlayData
+	* @param bMarkRenderStateDirty: if true, the change should be visible immediatelly. If you are updating many instances you should only set this to true for the last instance
+	*/
+	UFUNCTION(BlueprintCallable, Category = "AnimToTexture|Playback")
+	static bool UpdateInstanceAutoPlayData(UInstancedStaticMeshComponent* InstancedMeshComponent, int32 InstanceIndex, 
+		const FAnimToTextureAutoPlayData& AutoPlayData, bool bMarkRenderStateDirty=true);
+
+	/**
+	* * Updates a single instance with given FrameData
+	* @param bMarkRenderStateDirty: if true, the change should be visible immediatelly. If you are updating many instances you should only set this to true for the last instance
+	*/
+	UFUNCTION(BlueprintCallable, Category = "AnimToTexture|Playback")
+	static bool UpdateInstanceFrameData(UInstancedStaticMeshComponent* InstancedMeshComponent, int32 InstanceIndex, 
+		const FAnimToTextureFrameData& FrameData, bool bMarkRenderStateDirty=true);
+
+	/**
+	* Returns an AutoPlayData with the Start and End Frame for the given AnimationIndex.
+	* If AnimationIndex is out of range, false will be returned.
+	*/
+	UFUNCTION(BlueprintCallable, Category = "AnimToTexture|Playback")
+	static bool GetAutoPlayDataFromDataAsset(const UAnimToTextureDataAsset* DataAsset, int32 AnimationIndex,
+		FAnimToTextureAutoPlayData& AutoPlayData, float TimeOffset = 0.f, float PlayRate = 1.f);
+
+	UFUNCTION(BlueprintCallable, Category = "AnimToTexture|Playback")
+	static float GetFrame(float Time, float StartFrame, float EndFrame, 
+		float TimeOffset = 0.f, float PlayRate = 1.f, float SampleRate = 30.f);
+
+	UFUNCTION(BlueprintCallable, Category = "AnimToTexture|Playback")
+	static bool GetFrameDataFromDataAsset(const UAnimToTextureDataAsset* DataAsset, int32 AnimationIndex, float Time,
+		FAnimToTextureFrameData& AutoPlayData, float TimeOffset = 0.f, float PlayRate = 1.f);
+	
+private:
+	
+	template <class DataType>
+	static bool BatchUpdateInstancesData(UInstancedStaticMeshComponent* InstancedMeshComponent,
+		const TArray<DataType>& Data, const TArray<FMatrix>& Transforms, bool bMarkRenderStateDirty = true);
+	
+	template <class DataType>
+	static bool UpdateInstanceCustomData(UInstancedStaticMeshComponent* InstancedMeshComponent, int32 InstanceIndex,
+		const DataType& Data, bool bMarkRenderStateDirty = true);
+
+	template <class DataType>
+	static void GetData(const DataType& Data, TArray<float>& CustomData);
+
+	template <class DataType>
+	static void GetCustomData(const TArray<DataType>& Data, TArray<float>& CustomData);
+
+};
+
+
+template <class DataType>
+bool UAnimToTextureInstancePlaybackLibrary::BatchUpdateInstancesData(UInstancedStaticMeshComponent* InstancedMeshComponent,
+	const TArray<DataType>& Data, const TArray<FMatrix>& Transforms, bool bMarkRenderStateDirty)
+{
+	if (!InstancedMeshComponent)
+	{
+		return false;
+	}
+
+	// Get Number of Instances
+	const int32 NumInstances = InstancedMeshComponent->GetNumRenderInstances();
+
+	// Check number of instances
+	if (NumInstances &&
+		NumInstances == Data.Num() &&
+		NumInstances == Transforms.Num())
+	{
+		// Check data size
+		if (InstancedMeshComponent->NumCustomDataFloats == Data.GetTypeSize() / sizeof(float))
+		{
+			SIZE_T CustomDataSizeToCopy = InstancedMeshComponent->PerInstanceSMCustomData.Num() * InstancedMeshComponent->PerInstanceSMCustomData.GetTypeSize();
+			FMemory::Memcpy(InstancedMeshComponent->PerInstanceSMCustomData.GetData(), Data.GetData(), CustomDataSizeToCopy);
+
+			// FInstancedStaticMeshInstanceData is not exposed to blueprints, so we are duplicating the data here :/
+			TArray<FInstancedStaticMeshInstanceData> InstanceData;
+			InstanceData.AddUninitialized(NumInstances);
+			FMemory::Memcpy(InstanceData.GetData(), Transforms.GetData(), NumInstances * InstanceData.GetTypeSize());
+
+			// Batch Update
+			return InstancedMeshComponent->BatchUpdateInstancesData(0, NumInstances, InstanceData.GetData(), bMarkRenderStateDirty, false);
+		}
+	}
+
+	return false;
+}
+
+template <class DataType>
+void UAnimToTextureInstancePlaybackLibrary::GetData(const DataType& Data, TArray<float>& CustomData)
+{
+	const SIZE_T DataSize = sizeof(Data);
+	CustomData.SetNumUninitialized(DataSize / sizeof(float));
+
+	FMemory::Memcpy(CustomData.GetData(), &Data, DataSize);
+}
+
+template <class DataType>
+void UAnimToTextureInstancePlaybackLibrary::GetCustomData(const TArray<DataType>& Data, TArray<float>& CustomData)
+{
+	const int32 NumCustomDataFloats = Data.GetTypeSize() / sizeof(float);
+	CustomData.SetNumUninitialized(Data.Num() * NumCustomDataFloats);
+	
+	FMemory::Memcpy(CustomData.GetData(), Data.GetData(), Data.Num() * Data.GetTypeSize());
+}
+
+template <class DataType>
+bool UAnimToTextureInstancePlaybackLibrary::UpdateInstanceCustomData(UInstancedStaticMeshComponent* InstancedMeshComponent, int32 InstanceIndex,
+	const DataType& Data, bool bMarkRenderStateDirty)
+{
+	if (InstancedMeshComponent && InstancedMeshComponent->IsValidInstance(InstanceIndex))
+	{
+		// Check if valid size
+		const int32 NumCustomDataFloats = sizeof(Data) / sizeof(float);
+		if (InstancedMeshComponent->NumCustomDataFloats == NumCustomDataFloats)
+		{
+			TArray<float> CustomData;
+			GetData(Data, CustomData);
+			return InstancedMeshComponent->SetCustomData(InstanceIndex, CustomData, bMarkRenderStateDirty);
+		}
+	}
+
+	return false;
+}
diff --git a/Engine/Plugins/Experimental/AnimToTexture/Source/AnimToTexture/Public/AnimToTextureModule.h b/Engine/Plugins/Experimental/AnimToTexture/Source/AnimToTexture/Public/AnimToTextureModule.h
new file mode 100644
--- /dev/null
+++ b/Engine/Plugins/Experimental/AnimToTexture/Source/AnimToTexture/Public/AnimToTextureModule.h
@@ -0,0 +1,16 @@
+// Copyright Epic Games, Inc. All Rights Reserved.
+
+#pragma once
+
+#include "Modules/ModuleManager.h"
+
+DECLARE_LOG_CATEGORY_EXTERN(LogAnimToTexture, Log, All);
+
+class FAnimToTextureModule : public IModuleInterface
+{
+public:
+
+	/** IModuleInterface implementation */
+	virtual void StartupModule() override;
+	virtual void ShutdownModule() override;
+};
diff --git a/Engine/Plugins/Experimental/AnimToTexture/Source/AnimToTextureEditor/AnimToTextureEditor.Build.cs b/Engine/Plugins/Experimental/AnimToTexture/Source/AnimToTextureEditor/AnimToTextureEditor.Build.cs
new file mode 100644
--- /dev/null
+++ b/Engine/Plugins/Experimental/AnimToTexture/Source/AnimToTextureEditor/AnimToTextureEditor.Build.cs
@@ -0,0 +1,58 @@
+// Copyright Epic Games, Inc. All Rights Reserved.
+
+using UnrealBuildTool;
+
+public class AnimToTextureEditor : ModuleRules
+{
+	public AnimToTextureEditor(ReadOnlyTargetRules Target) : base(Target)
+	{
+		PCHUsage = ModuleRules.PCHUsageMode.UseExplicitOrSharedPCHs;
+		
+		PublicIncludePaths.AddRange(
+			new string[] {
+				// ... add public include paths required here ...
+			}
+			);
+				
+		
+		PrivateIncludePaths.AddRange(
+			new string[] {
+				// ... add other private include paths required here ...
+			}
+			);
+			
+		
+		PublicDependencyModuleNames.AddRange(
+			new string[]
+			{
+				"Core",
+				"MaterialEditor"
+				// ... add other public dependencies that you statically link with here ...
+			}
+			);
+			
+		
+		PrivateDependencyModuleNames.AddRange(
+			new string[]
+			{
+				"CoreUObject",
+				"Engine",
+				"Slate",
+				"SlateCore",
+				"AnimToTexture",
+				"RawMesh",
+				"MeshDescription",
+				"StaticMeshDescription",
+				"UnrealEd",
+			}
+			);
+		
+		
+		DynamicallyLoadedModuleNames.AddRange(
+			new string[]
+			{
+				// ... add any modules that your module loads dynamically here ...
+			}
+			);
+	}
+}
diff --git a/Engine/Plugins/Experimental/AnimToTexture/Source/AnimToTextureEditor/Private/AnimToTextureBPLibrary.cpp b/Engine/Plugins/Experimental/AnimToTexture/Source/AnimToTextureEditor/Private/AnimToTextureBPLibrary.cpp
new file mode 100644
--- /dev/null
+++ b/Engine/Plugins/Experimental/AnimToTexture/Source/AnimToTextureEditor/Private/AnimToTextureBPLibrary.cpp
@@ -0,0 +1,1148 @@
+// Copyright Epic Games, Inc. All Rights Reserved.
+
+#include "AnimToTextureBPLibrary.h"
+#include "AnimToTextureEditorModule.h"
+#include "AnimToTextureUtils.h"
+#include "AnimToTextureSkeletalMesh.h"
+
+#include "LevelEditor.h"
+#include "RawMesh.h"
+#include "MeshUtilities.h"
+#include "Modules/ModuleManager.h"
+#include "Engine/SkeletalMesh.h"
+#include "Components/SkeletalMeshComponent.h"
+#include "Animation/Skeleton.h"
+#include "Animation/AnimSequence.h"
+#include "Math/Vector.h"
+#include "Math/NumericLimits.h"
+#include "MeshDescription.h"
+#include "Materials/MaterialInstanceConstant.h"
+#include "MaterialEditingLibrary.h"
+
+using namespace AnimToTexture_Private;
+
+UAnimToTextureBPLibrary::UAnimToTextureBPLibrary(const FObjectInitializer& ObjectInitializer)
+: Super(ObjectInitializer)
+{
+
+}
+
+void UAnimToTextureBPLibrary::AnimationToTexture(UAnimToTextureDataAsset* DataAsset, const FTransform RootTransform, bool& bSuccess)
+{
+	bSuccess = false; 
+
+	if (!DataAsset)
+	{
+		return;
+	}
+
+	// Reset DataAsset Info Values
+	DataAsset->ResetInfo();
+
+	// Check StaticMesh
+	if (!DataAsset->GetStaticMesh())
+	{
+		UE_LOG(LogAnimToTextureEditor, Warning, TEXT("Invalid StaticMesh"));
+		return;
+	}
+
+	// Check SkeletalMesh
+	if (!DataAsset->GetSkeletalMesh())
+	{
+		UE_LOG(LogAnimToTextureEditor, Warning, TEXT("Invalid SkeletalMesh"));
+		return;
+	}
+
+	// Check Skeleton
+	if (!DataAsset->GetSkeletalMesh()->GetSkeleton())
+	{
+		UE_LOG(LogAnimToTextureEditor, Warning, TEXT("Invalid SkeletalMesh. No valid Skeleton found"));
+		return;
+	}
+
+	// Check StaticMesh LOD
+	if (!DataAsset->GetStaticMesh()->IsSourceModelValid(DataAsset->StaticLODIndex))
+	{
+		UE_LOG(LogAnimToTextureEditor, Warning, TEXT("Invalid StaticMesh LOD Index: %i"), DataAsset->StaticLODIndex);
+		return;
+	}
+
+	// Check SkeletalMesh LOD
+	if (!DataAsset->GetSkeletalMesh()->IsValidLODIndex(DataAsset->SkeletalLODIndex))
+	{
+		UE_LOG(LogAnimToTextureEditor, Warning, TEXT("Invalid SkeletalMesh LOD Index: %i"), DataAsset->SkeletalLODIndex);
+		return;
+	}
+
+	// Check Socket.
+	bool bValidSocket = false;
+	if (DataAsset->AttachToSocket.IsValid() && !DataAsset->AttachToSocket.IsNone())
+	{
+		if (HasBone(DataAsset->GetSkeletalMesh(), DataAsset->AttachToSocket))
+		{
+			bValidSocket = true;
+		}
+		else
+		{
+			UE_LOG(LogAnimToTextureEditor, Warning, TEXT("Invalid Socket: %s"), *DataAsset->AttachToSocket.ToString());
+			return;
+		}
+	}
+	if (bValidSocket && DataAsset->Mode == EAnimToTextureMode::Vertex)
+	{
+		UE_LOG(LogAnimToTextureEditor, Warning, TEXT("Unable to use Socket in Vertex Mode. Use Bone Mode instead."));
+		return;
+	}
+
+	// Check if UVChannel is being used by the Lightmap UV
+	const FStaticMeshSourceModel& SourceModel = DataAsset->GetStaticMesh()->GetSourceModel(DataAsset->StaticLODIndex);
+	if (SourceModel.BuildSettings.bGenerateLightmapUVs &&
+		SourceModel.BuildSettings.DstLightmapIndex == DataAsset->UVChannel)
+	{
+		UE_LOG(LogAnimToTextureEditor, Warning, TEXT("Invalid UVChannel: %i. Already used by LightMap"), DataAsset->UVChannel);
+		return;
+	}
+
+	// Check Animations
+	int32 NumAnimations = 0;
+	for (const FAnimToTextureAnimSequenceInfo& AnimSequenceInfo : DataAsset->AnimSequences)
+	{
+		const UAnimSequence* AnimSequence = AnimSequenceInfo.AnimSequence;
+		if (AnimSequenceInfo.bEnabled && AnimSequence)
+		{
+			// Check Frame Range
+			if (AnimSequenceInfo.bUseCustomRange && 
+				(AnimSequenceInfo.EndFrame - AnimSequenceInfo.StartFrame) <= 0 )
+			{ 
+				UE_LOG(LogAnimToTextureEditor, Warning, TEXT("Invalid Custom Range for AnimSequence: %s"), *AnimSequence->GetName());
+				return;
+			}
+
+			NumAnimations++;
+		}
+	}
+	if (!NumAnimations)
+	{
+		UE_LOG(LogAnimToTextureEditor, Warning, TEXT("No Animations found"));
+		return;
+	};
+
+	// ---------------------------------------------------------------------------		
+	// Get Meshes Vertices and Mapping.
+	// NOTE: We need to create a Mapping between the StaticMesh and the SkeletalMesh
+	//       Since they dont have same number of points.
+	//
+
+	// Get Static Vertices 
+	// NOTE: these are the RenderData Vertices.
+	TArray<FVector3f> Vertices;
+	TArray<FVector3f> Normals;
+	const int32 NumVertices = GetVertices(DataAsset->GetStaticMesh(), DataAsset->StaticLODIndex, Vertices, Normals);
+
+	// Get SkeletalMesh Vertices 
+	TArray<FVector3f> SkelVertices;
+	GetVertices(DataAsset->GetSkeletalMesh(), DataAsset->SkeletalLODIndex, SkelVertices);
+	
+	// Get Mapping
+	TArray<FVertexToMeshMapping> Mapping;
+	FVertexToMeshMapping::Create(
+		DataAsset->GetStaticMesh(), DataAsset->StaticLODIndex,
+		DataAsset->GetSkeletalMesh(), DataAsset->SkeletalLODIndex, Mapping);
+
+	// --------------------------------------------------------------------------
+		
+	// Create Temp Actor
+	check(GEditor);
+	UWorld* World = GEditor->GetEditorWorldContext().World();
+	check(World);
+
+	AActor* Actor = World->SpawnActor<AActor>();
+	check(Actor);
+
+	// Create Temp SkeletalMesh Component
+	USkeletalMeshComponent* SkeletalMeshComponent = NewObject<USkeletalMeshComponent>(Actor);
+	check(SkeletalMeshComponent);
+	SkeletalMeshComponent->SetSkeletalMesh(DataAsset->GetSkeletalMesh());
+	SkeletalMeshComponent->SetForcedLOD(1); // Force to LOD0;
+	SkeletalMeshComponent->SetAnimationMode(EAnimationMode::AnimationSingleNode);
+	SkeletalMeshComponent->SetUpdateAnimationInEditor(true);
+	SkeletalMeshComponent->VisibilityBasedAnimTickOption = EVisibilityBasedAnimTickOption::AlwaysTickPoseAndRefreshBones;
+	SkeletalMeshComponent->RegisterComponent();
+
+	// ---------------------------------------------------------------------------
+	// Get Reference Skeleton Transforms
+	//
+	int32 NumBones = INDEX_NONE;
+	int32 SocketIndex = INDEX_NONE;
+	TArray<FName>     BoneNames;
+	TArray<FVector3f> BoneRefPositions;
+	TArray<FVector4>  BoneRefRotations;
+	TArray<FVector3f> BonePositions;
+	TArray<FVector4>  BoneRotations;
+	
+	if (DataAsset->Mode == EAnimToTextureMode::Bone)
+	{
+		// Gets Ref Bone Position and Rotations.
+		NumBones = GetRefBonePositionsAndRotations(SkeletalMeshComponent,
+			BoneRefPositions, BoneRefRotations);
+
+		// NOTE: there is a limitation with the number of bones atm.
+		if (NumBones > 256)
+		{
+			UE_LOG(LogAnimToTextureEditor, Warning, TEXT("Invalid Number of Bones. There is a maximum of 256 bones"))
+			return;
+		}
+
+		// Get Bone Names (no virtual)
+		GetBoneNames(SkeletalMeshComponent->GetSkeletalMeshAsset(), BoneNames);
+
+		// Make sure array sizes are correct.
+		check(BoneNames.Num() == NumBones);
+
+		// Check if Socket is in BoneNames
+		if (bValidSocket && !BoneNames.Find(DataAsset->AttachToSocket, SocketIndex))
+		{
+			UE_LOG(LogAnimToTextureEditor, Warning, TEXT("Socket: %s not found in Raw Bone List"), *DataAsset->AttachToSocket.ToString());
+			return;
+		}
+
+		// Add RefPose 
+		// Note: this is added in the first frame of the Bone Position and Rotation Textures
+		BonePositions.Append(BoneRefPositions);
+		BoneRotations.Append(BoneRefRotations);
+	}
+
+	// ---------------------------------------------------------------------------
+	// Get Vertex Data (for all frames)
+	//		
+	TArray<FVector3f> VertexDeltas;
+	TArray<FVector3f> VertexNormals;
+	
+	// Get Animation Frames Data
+	//
+	const float SampleInterval = 1.f / DataAsset->SampleRate;
+
+	for (const FAnimToTextureAnimSequenceInfo& AnimSequenceInfo : DataAsset->AnimSequences)
+	{
+		UAnimSequence* AnimSequence = AnimSequenceInfo.AnimSequence;
+
+		if (!AnimSequenceInfo.bEnabled || !AnimSequence)
+		{
+			continue;
+		}
+		
+		// Make sure SkeletalMesh is compatible with AnimSequence
+		if (!SkeletalMeshComponent->GetSkeletalMeshAsset()->GetSkeleton()->IsCompatibleForEditor(AnimSequence->GetSkeleton()))
+		{
+			UE_LOG(LogAnimToTextureEditor, Warning, TEXT("Invalid AnimSequence: %s for given SkeletalMesh: %s"), *AnimSequence->GetFName().ToString(), *SkeletalMeshComponent->GetSkeletalMeshAsset()->GetFName().ToString());
+			continue;
+		}
+		// Set AnimSequence
+		else
+		{			
+			SkeletalMeshComponent->SetAnimation(AnimSequence);
+		}
+
+		// -----------------------------------------------------------------------------------
+		// Get Number of Frames
+		//
+		int32 StartFrame;
+		int32 EndFrame;
+		
+		// Get Range from AnimSequence
+		if (!AnimSequenceInfo.bUseCustomRange)
+		{
+			StartFrame = 0;
+			EndFrame = AnimSequence->GetNumberOfSampledKeys() - 1; // AnimSequence->GetNumberOfFrames();
+		}
+		// Get Range from DataAsset
+		else
+		{
+			StartFrame = AnimSequenceInfo.StartFrame;
+			EndFrame = AnimSequenceInfo.EndFrame;
+		}
+		if (EndFrame - StartFrame <= 0)
+		{
+			continue;
+		}
+		
+		// ---------------------------------------------------------------------------
+		// 
+		float Time = 0.f;
+		const float EndTime = AnimSequence->GetTimeAtFrame(EndFrame);
+
+		int32 SampleIndex = 0;
+
+		while (Time < EndTime)
+		{
+			Time = FMath::Clamp(SampleIndex * SampleInterval, 0.f, EndTime);
+			SampleIndex++;
+
+			// Go To Time
+			SkeletalMeshComponent->SetPosition(Time);
+
+			// Update SkelMesh Animation.
+			SkeletalMeshComponent->TickAnimation(0.f, false /*bNeedsValidRootMotion*/);
+			// SkeletalMeshComponent->TickComponent(0.f, ELevelTick::LEVELTICK_All, nullptr);
+			SkeletalMeshComponent->RefreshBoneTransforms(nullptr /*TickFunction*/);
+			
+			// ---------------------------------------------------------------------------
+			// Store Vertex Deltas & Normals.
+			//
+			if (DataAsset->Mode == EAnimToTextureMode::Vertex)
+			{
+				TArray<FVector3f> VertexFrameDeltas;
+				TArray<FVector3f> VertexFrameNormals;
+				
+				GetVertexDeltasAndNormals(SkeletalMeshComponent, DataAsset->SkeletalLODIndex,
+					Vertices, Normals, Mapping,
+					RootTransform,
+					VertexFrameDeltas, VertexFrameNormals);
+					
+				VertexDeltas.Append(VertexFrameDeltas);
+				VertexNormals.Append(VertexFrameNormals);
+
+			} // End Vertex Mode
+
+			// ---------------------------------------------------------------------------
+			// Store Bone Positions & Rotations
+			//
+			else if (DataAsset->Mode == EAnimToTextureMode::Bone)
+			{
+				TArray<FVector3f> BoneFramePositions;
+				TArray<FVector4> BoneFrameRotations;
+
+				GetBonePositionsAndRotations(SkeletalMeshComponent, BoneRefPositions,
+					BoneFramePositions, BoneFrameRotations);
+
+				BonePositions.Append(BoneFramePositions);
+				BoneRotations.Append(BoneFrameRotations);
+
+			} // End Bone Mode
+
+		} // End Frame
+
+		// Store Anim Info Data
+		FAnimToTextureAnimInfo AnimInfo;
+		AnimInfo.StartFrame = DataAsset->NumFrames;
+		AnimInfo.EndFrame = DataAsset->NumFrames + EndFrame;
+		DataAsset->Animations.Add(AnimInfo);
+
+		// Accumulate Frames
+		DataAsset->NumFrames += SampleIndex;
+	} // End Anim
+		
+	// Destroy Temp Component
+	SkeletalMeshComponent->UnregisterComponent();
+	SkeletalMeshComponent->DestroyComponent();
+
+	Actor->Destroy();
+	
+	// ---------------------------------------------------------------------------
+	// Nothing to do here ...
+	//
+	if (!DataAsset->NumFrames || !NumVertices)
+	{
+		return;
+	}
+
+	// ---------------------------------------------------------------------------
+	if (DataAsset->Mode == EAnimToTextureMode::Vertex)
+	{
+		// Find Best Resolution for Vertex Data
+		int32 Height, Width;
+		if (!FindBestResolution(DataAsset->NumFrames, NumVertices, 
+								Height, Width, DataAsset->VertexRowsPerFrame, 
+								DataAsset->MaxHeight, DataAsset->MaxWidth, DataAsset->bEnforcePowerOfTwo))
+		{
+			UE_LOG(LogAnimToTextureEditor, Warning, TEXT("Vertex Animation data cannot be fit in a %ix%i texture."), DataAsset->MaxHeight, DataAsset->MaxWidth);
+			return;
+		}
+
+		// Normalize Vertex Data
+		TArray<FVector3f> NormalizedVertexDeltas;
+		TArray<FVector3f> NormalizedVertexNormals;
+		NormalizeVertexData(
+			VertexDeltas, VertexNormals,
+			DataAsset->VertexMinBBox, DataAsset->VertexSizeBBox,
+			NormalizedVertexDeltas, NormalizedVertexNormals);
+
+		// Write Textures
+		if (DataAsset->Precision == EAnimToTexturePrecision::SixteenBits)
+		{
+			WriteVectorsToTexture<FVector3f, FHighPrecision>(NormalizedVertexDeltas, DataAsset->NumFrames, DataAsset->VertexRowsPerFrame, Height, Width, DataAsset->GetVertexPositionTexture());
+			WriteVectorsToTexture<FVector3f, FHighPrecision>(NormalizedVertexNormals, DataAsset->NumFrames, DataAsset->VertexRowsPerFrame, Height, Width, DataAsset->GetVertexNormalTexture());
+		}
+		else
+		{
+			WriteVectorsToTexture<FVector3f, FLowPrecision>(NormalizedVertexDeltas, DataAsset->NumFrames, DataAsset->VertexRowsPerFrame, Height, Width, DataAsset->GetVertexPositionTexture());
+			WriteVectorsToTexture<FVector3f, FLowPrecision>(NormalizedVertexNormals, DataAsset->NumFrames, DataAsset->VertexRowsPerFrame, Height, Width, DataAsset->GetVertexNormalTexture());
+		}		
+
+		// Add Vertex UVChannel
+		CreateUVChannel(DataAsset->GetStaticMesh(), DataAsset->StaticLODIndex, DataAsset->UVChannel, Height, Width);
+
+		// Update Bounds
+		SetBoundsExtensions(DataAsset->GetStaticMesh(), DataAsset->VertexMinBBox, DataAsset->VertexSizeBBox);
+
+		// Done with StaticMesh
+		DataAsset->GetStaticMesh()->PostEditChange();
+	}
+
+	// ---------------------------------------------------------------------------
+	
+	if (DataAsset->Mode == EAnimToTextureMode::Bone)
+	{
+		// Find Best Resolution for Bone Data
+		int32 Height, Width;
+
+		// Note we are adding +1 frame for the ref pose
+		if (!FindBestResolution(DataAsset->NumFrames + 1, NumBones, 
+								Height, Width, DataAsset->BoneRowsPerFrame, 
+								DataAsset->MaxHeight, DataAsset->MaxWidth, DataAsset->bEnforcePowerOfTwo))
+		{
+			UE_LOG(LogAnimToTextureEditor, Warning, TEXT("Bone Animation data cannot be fit in a %ix%i texture."), DataAsset->MaxHeight, DataAsset->MaxWidth);
+			return;
+		}
+
+		// Write Bone Position and Rotation Textures
+		{
+			// Normalize Bone Data
+			TArray<FVector3f> NormalizedBonePositions;
+			TArray<FVector4> NormalizedBoneRotations;
+			NormalizeBoneData(
+				BonePositions, BoneRotations,
+				DataAsset->BoneMinBBox, DataAsset->BoneSizeBBox,
+				NormalizedBonePositions, NormalizedBoneRotations);
+
+			// Write Textures
+			if (DataAsset->Precision == EAnimToTexturePrecision::SixteenBits)
+			{
+				WriteVectorsToTexture<FVector3f, FHighPrecision>(NormalizedBonePositions, DataAsset->NumFrames + 1, DataAsset->BoneRowsPerFrame, Height, Width, DataAsset->GetBonePositionTexture());
+				WriteVectorsToTexture<FVector4, FHighPrecision>(NormalizedBoneRotations, DataAsset->NumFrames + 1, DataAsset->BoneRowsPerFrame, Height, Width, DataAsset->GetBoneRotationTexture());
+			}
+			else
+			{
+				WriteVectorsToTexture<FVector3f, FLowPrecision>(NormalizedBonePositions, DataAsset->NumFrames + 1, DataAsset->BoneRowsPerFrame, Height, Width, DataAsset->GetBonePositionTexture());
+				WriteVectorsToTexture<FVector4, FLowPrecision>(NormalizedBoneRotations, DataAsset->NumFrames + 1, DataAsset->BoneRowsPerFrame, Height, Width, DataAsset->GetBoneRotationTexture());
+			}
+		}
+
+		// ---------------------------------------------------------------------------
+
+		// Find Best Resolution for Weights Data
+		if (!FindBestResolution(2, NumVertices, 
+								Height, Width, DataAsset->BoneWeightRowsPerFrame, 
+								DataAsset->MaxHeight, DataAsset->MaxWidth, DataAsset->bEnforcePowerOfTwo))
+		{
+			UE_LOG(LogAnimToTextureEditor, Warning, TEXT("Weights Data cannot be fit in a %ix%i texture."), DataAsset->MaxHeight, DataAsset->MaxWidth);
+			return;
+		}
+
+		// Write Weights Texture
+		{
+			TArray<TVertexSkinWeight<4>> SkinWeights;
+
+			// Store Influence Weights
+			if (!bValidSocket)
+			{
+				// Get SkeletalMesh Weights
+				TArray<VertexSkinWeightMax> SkeletalMeshSkinWeights;
+				GetSkinWeights(DataAsset->GetSkeletalMesh(), DataAsset->SkeletalLODIndex, SkeletalMeshSkinWeights);
+
+				// Interpolate Weights with Mapping
+				TArray<VertexSkinWeightMax> StaticMeshSkinWeights;
+				FVertexToMeshMapping::InterpolateSkinWeights(Mapping, SkeletalMeshSkinWeights, StaticMeshSkinWeights);
+				
+				// Reduce SkeletalMesh Weights to 4 highest influences.
+				ReduceSkinWeights(StaticMeshSkinWeights, SkinWeights);
+			}
+
+			// If Valid Socket, set all influences to same index.
+			else
+			{
+				// Set all indices and weights to same SocketIndex
+				SkinWeights.SetNumUninitialized(NumVertices);
+				for (TVertexSkinWeight<4>& SkinWeight : SkinWeights)
+				{
+					SkinWeight.BoneWeights = TStaticArray<uint8, 4>(InPlace, 255);
+					SkinWeight.MeshBoneIndices = TStaticArray<uint16, 4>(InPlace, SocketIndex);
+				}
+			}
+
+			// Write Bone Weights Texture
+			WriteSkinWeightsToTexture(SkinWeights,
+				DataAsset->BoneWeightRowsPerFrame, Height, Width, DataAsset->GetBoneWeightTexture());
+		}
+
+		// Add Vertex UVChannel
+		CreateUVChannel(DataAsset->GetStaticMesh(), DataAsset->StaticLODIndex, DataAsset->UVChannel, Height, Width);
+
+		// Update Bounds
+		SetBoundsExtensions(DataAsset->GetStaticMesh(), DataAsset->BoneMinBBox, DataAsset->BoneSizeBBox);
+
+		// Done with StaticMesh
+		DataAsset->GetStaticMesh()->PostEditChange();
+	}
+
+	// ---------------------------------------------------------------------------
+	// Mark Packages dirty
+	//
+	DataAsset->MarkPackageDirty();
+	
+	// All good here !
+	bSuccess = true;
+}
+
+
+// 
+void UAnimToTextureBPLibrary::GetVertexDeltasAndNormals(const USkeletalMeshComponent* SkeletalMeshComponent, const int32 LODIndex, 
+	const TArray<FVector3f>& Vertices, const TArray<FVector3f>& Normals,  const TArray<FVertexToMeshMapping>& Mapping,
+	const FTransform RootTransform,
+	TArray<FVector3f>& OutVertexDeltas, TArray<FVector3f>& OutVertexNormals)
+{
+	OutVertexDeltas.Reset();
+	OutVertexNormals.Reset();
+
+	check(Vertices.Num() == Mapping.Num());
+	const int32 NumVertices = Vertices.Num();
+
+	// Get Deformed vertices at current frame
+	TArray<FVector3f> SkinnedVertices;
+	GetSkinnedVertices(SkeletalMeshComponent, LODIndex, SkinnedVertices);
+
+	// Allocate
+	OutVertexDeltas.SetNumUninitialized(NumVertices);
+	OutVertexNormals.SetNumUninitialized(NumVertices);
+
+	// Loop thru static vertices and find the mapped SkeletalMesh Vertex
+	for (int32 VertexIndex = 0; VertexIndex < NumVertices; VertexIndex++)
+	{
+		// Get Static/SkeletalMesh Mapping
+		const FVertexToMeshMapping& VertexMapping = Mapping[VertexIndex];
+
+		// Transform Vertex
+		// (assuming indices are valid)
+		const FIntVector3& Triangle = VertexMapping.Triangle;
+		const FVector3f& A = SkinnedVertices[Triangle.X];
+		const FVector3f& B = SkinnedVertices[Triangle.Y];
+		const FVector3f& C = SkinnedVertices[Triangle.Z];
+
+		const FVector3f VertexPosition = VertexMapping.TransformPosition(Vertices[VertexIndex], A, B, C);
+		const FVector3f VertexNormal = VertexMapping.TransformVector(Normals[VertexIndex], A, B, C);
+
+		// Transform Position and Delta with RootTransform
+		const FVector3f TransformedVertexDelta = ((FVector3f)RootTransform.TransformPosition((FVector)VertexPosition)) - Vertices[VertexIndex];
+		const FVector3f TransformedVertexNormal = (FVector3f)RootTransform.TransformVector((FVector)VertexNormal);
+		
+		OutVertexDeltas[VertexIndex] = TransformedVertexDelta;
+		OutVertexNormals[VertexIndex] = TransformedVertexNormal;
+	}
+}
+
+
+int32 UAnimToTextureBPLibrary::GetRefBonePositionsAndRotations(const USkeletalMeshComponent* SkeletalMeshComponent,
+	TArray<FVector3f>& OutBoneRefPositions, TArray<FVector4>& OutBoneRefRotations)
+{
+	OutBoneRefPositions.Reset();
+	OutBoneRefRotations.Reset();
+
+	// Get Number of RawBones (no virtual)
+	const int32 NumBones = GetNumBones(SkeletalMeshComponent->GetSkeletalMeshAsset());
+	
+	// Get Raw Ref Bone (no virtual)
+	TArray<FTransform> RefBoneTransforms;
+	GetRefBoneTransforms(SkeletalMeshComponent->GetSkeletalMeshAsset(), RefBoneTransforms);
+	DecomposeTransformations(RefBoneTransforms, OutBoneRefPositions, OutBoneRefRotations);
+
+	return NumBones;
+}
+
+
+int32 UAnimToTextureBPLibrary::GetBonePositionsAndRotations(const USkeletalMeshComponent* SkeletalMeshComponent, const TArray<FVector3f>& BoneRefPositions,
+	TArray<FVector3f>& BonePositions, TArray<FVector4>& BoneRotations)
+{
+	BonePositions.Reset();
+	BoneRotations.Reset();
+
+	// Get Relative Transforms
+	// Note: Size is of Raw bones in SkeletalMesh. These are the original/raw bones of the asset, without Virtual Bones.
+	TArray<FMatrix44f> RefToLocals;
+	SkeletalMeshComponent->CacheRefToLocalMatrices(RefToLocals);
+	const int32 NumBones = RefToLocals.Num();
+
+	// check size
+	check(NumBones == BoneRefPositions.Num());
+
+	// Get Component Space Transforms
+	// Note returns all transforms, including VirtualBones
+	const TArray<FTransform>& CompSpaceTransforms = SkeletalMeshComponent->GetComponentSpaceTransforms();
+	check(CompSpaceTransforms.Num() >= RefToLocals.Num());
+
+	// Allocate
+	BonePositions.SetNumUninitialized(NumBones);
+	BoneRotations.SetNumUninitialized(NumBones);
+
+	for (int32 BoneIndex = 0; BoneIndex < NumBones; BoneIndex++)
+	{
+		// Decompose Transformation (ComponentSpace)
+		const FTransform& CompSpaceTransform = CompSpaceTransforms[BoneIndex];
+		FVector3f BonePosition;
+		FVector4 BoneRotation;
+		DecomposeTransformation(CompSpaceTransform, BonePosition, BoneRotation);
+
+		// Position Delta (from RefPose)
+		const FVector3f Delta = BonePosition - BoneRefPositions[BoneIndex];
+
+		// Decompose Transformation (Relative to RefPose)
+		FVector3f BoneRelativePosition;
+		FVector4 BoneRelativeRotation;
+		const FMatrix RefToLocalMatrix(RefToLocals[BoneIndex]);
+		const FTransform RelativeTransform(RefToLocalMatrix);
+		DecomposeTransformation(RelativeTransform, BoneRelativePosition, BoneRelativeRotation);
+
+		BonePositions[BoneIndex] = Delta;
+		BoneRotations[BoneIndex] = BoneRelativeRotation;
+	}
+
+	return NumBones;
+}
+
+
+void UAnimToTextureBPLibrary::UpdateMaterialInstanceFromDataAsset(UAnimToTextureDataAsset* DataAsset, UMaterialInstanceConstant* MaterialInstance, 
+	const bool bAutoPlay, const int32 AnimationIndex, 
+	const EAnimToTextureNumBoneInfluences NumBoneInfluences, const EMaterialParameterAssociation MaterialParameterAssociation)
+{
+	if (!MaterialInstance || !DataAsset)
+	{
+		return;
+	}
+
+	// Set UVChannel
+	switch (DataAsset->UVChannel)
+	{
+		case 0:
+			UMaterialEditingLibrary::SetMaterialInstanceStaticSwitchParameterValue(MaterialInstance, AnimToTextureParamNames::UseUV0, true, MaterialParameterAssociation);
+			UMaterialEditingLibrary::SetMaterialInstanceStaticSwitchParameterValue(MaterialInstance, AnimToTextureParamNames::UseUV1, false, MaterialParameterAssociation);
+			UMaterialEditingLibrary::SetMaterialInstanceStaticSwitchParameterValue(MaterialInstance, AnimToTextureParamNames::UseUV2, false, MaterialParameterAssociation);
+			UMaterialEditingLibrary::SetMaterialInstanceStaticSwitchParameterValue(MaterialInstance, AnimToTextureParamNames::UseUV3, false, MaterialParameterAssociation);
+			break;
+		case 1:
+			UMaterialEditingLibrary::SetMaterialInstanceStaticSwitchParameterValue(MaterialInstance, AnimToTextureParamNames::UseUV0, false, MaterialParameterAssociation);
+			UMaterialEditingLibrary::SetMaterialInstanceStaticSwitchParameterValue(MaterialInstance, AnimToTextureParamNames::UseUV1, true, MaterialParameterAssociation);
+			UMaterialEditingLibrary::SetMaterialInstanceStaticSwitchParameterValue(MaterialInstance, AnimToTextureParamNames::UseUV2, false, MaterialParameterAssociation);
+			UMaterialEditingLibrary::SetMaterialInstanceStaticSwitchParameterValue(MaterialInstance, AnimToTextureParamNames::UseUV3, false, MaterialParameterAssociation);
+			break;
+		case 2:
+			UMaterialEditingLibrary::SetMaterialInstanceStaticSwitchParameterValue(MaterialInstance, AnimToTextureParamNames::UseUV0, false, MaterialParameterAssociation);
+			UMaterialEditingLibrary::SetMaterialInstanceStaticSwitchParameterValue(MaterialInstance, AnimToTextureParamNames::UseUV1, false, MaterialParameterAssociation);
+			UMaterialEditingLibrary::SetMaterialInstanceStaticSwitchParameterValue(MaterialInstance, AnimToTextureParamNames::UseUV2, true, MaterialParameterAssociation);
+			UMaterialEditingLibrary::SetMaterialInstanceStaticSwitchParameterValue(MaterialInstance, AnimToTextureParamNames::UseUV3, false, MaterialParameterAssociation);
+			break;
+		case 3:
+			UMaterialEditingLibrary::SetMaterialInstanceStaticSwitchParameterValue(MaterialInstance, AnimToTextureParamNames::UseUV0, false, MaterialParameterAssociation);
+			UMaterialEditingLibrary::SetMaterialInstanceStaticSwitchParameterValue(MaterialInstance, AnimToTextureParamNames::UseUV1, false, MaterialParameterAssociation);
+			UMaterialEditingLibrary::SetMaterialInstanceStaticSwitchParameterValue(MaterialInstance, AnimToTextureParamNames::UseUV2, false, MaterialParameterAssociation);
+			UMaterialEditingLibrary::SetMaterialInstanceStaticSwitchParameterValue(MaterialInstance, AnimToTextureParamNames::UseUV3, true, MaterialParameterAssociation);
+			break;
+		default:
+			UMaterialEditingLibrary::SetMaterialInstanceStaticSwitchParameterValue(MaterialInstance, AnimToTextureParamNames::UseUV0, false, MaterialParameterAssociation);
+			UMaterialEditingLibrary::SetMaterialInstanceStaticSwitchParameterValue(MaterialInstance, AnimToTextureParamNames::UseUV1, true, MaterialParameterAssociation);
+			UMaterialEditingLibrary::SetMaterialInstanceStaticSwitchParameterValue(MaterialInstance, AnimToTextureParamNames::UseUV2, false, MaterialParameterAssociation);
+			UMaterialEditingLibrary::SetMaterialInstanceStaticSwitchParameterValue(MaterialInstance, AnimToTextureParamNames::UseUV3, false, MaterialParameterAssociation);
+			break;
+	}
+
+	// Update Vertex Params
+	if (DataAsset->Mode == EAnimToTextureMode::Vertex)
+	{
+		FLinearColor VectorParameter;
+		VectorParameter = FLinearColor(DataAsset->VertexMinBBox);
+		UMaterialEditingLibrary::SetMaterialInstanceVectorParameterValue(MaterialInstance, AnimToTextureParamNames::BoundingBoxMin, VectorParameter, MaterialParameterAssociation);
+		
+		VectorParameter = FLinearColor(DataAsset->VertexSizeBBox);
+		UMaterialEditingLibrary::SetMaterialInstanceVectorParameterValue(MaterialInstance, AnimToTextureParamNames::BoundingBoxScale, VectorParameter, MaterialParameterAssociation);
+		UMaterialEditingLibrary::SetMaterialInstanceScalarParameterValue(MaterialInstance, AnimToTextureParamNames::RowsPerFrame, DataAsset->VertexRowsPerFrame, MaterialParameterAssociation);
+
+		UMaterialEditingLibrary::SetMaterialInstanceTextureParameterValue(MaterialInstance, AnimToTextureParamNames::VertexPositionTexture, DataAsset->GetVertexPositionTexture(), MaterialParameterAssociation);
+		UMaterialEditingLibrary::SetMaterialInstanceTextureParameterValue(MaterialInstance, AnimToTextureParamNames::VertexNormalTexture, DataAsset->GetVertexNormalTexture(), MaterialParameterAssociation);
+
+	}
+
+	// Update Bone Params
+	else if (DataAsset->Mode == EAnimToTextureMode::Bone)
+	{
+		FLinearColor VectorParameter;
+		VectorParameter = FLinearColor(DataAsset->BoneMinBBox);
+		UMaterialEditingLibrary::SetMaterialInstanceVectorParameterValue(MaterialInstance, AnimToTextureParamNames::BoundingBoxMin, VectorParameter, MaterialParameterAssociation);
+
+		VectorParameter = FLinearColor(DataAsset->BoneSizeBBox);
+		UMaterialEditingLibrary::SetMaterialInstanceVectorParameterValue(MaterialInstance, AnimToTextureParamNames::BoundingBoxScale, VectorParameter, MaterialParameterAssociation);
+		UMaterialEditingLibrary::SetMaterialInstanceScalarParameterValue(MaterialInstance, AnimToTextureParamNames::RowsPerFrame, DataAsset->BoneRowsPerFrame, MaterialParameterAssociation);
+		UMaterialEditingLibrary::SetMaterialInstanceScalarParameterValue(MaterialInstance, AnimToTextureParamNames::BoneWeightRowsPerFrame, DataAsset->BoneWeightRowsPerFrame, MaterialParameterAssociation);
+
+		UMaterialEditingLibrary::SetMaterialInstanceTextureParameterValue(MaterialInstance, AnimToTextureParamNames::BonePositionTexture, DataAsset->GetBonePositionTexture(), MaterialParameterAssociation);
+		UMaterialEditingLibrary::SetMaterialInstanceTextureParameterValue(MaterialInstance, AnimToTextureParamNames::BoneRotationTexture, DataAsset->GetBoneRotationTexture(), MaterialParameterAssociation);
+		UMaterialEditingLibrary::SetMaterialInstanceTextureParameterValue(MaterialInstance, AnimToTextureParamNames::BoneWeightsTexture, DataAsset->GetBoneWeightTexture(), MaterialParameterAssociation);
+
+		// Num Influences
+		switch (NumBoneInfluences)
+		{
+			case EAnimToTextureNumBoneInfluences::One:
+				UMaterialEditingLibrary::SetMaterialInstanceStaticSwitchParameterValue(MaterialInstance, AnimToTextureParamNames::UseTwoInfluences, false, MaterialParameterAssociation);
+				UMaterialEditingLibrary::SetMaterialInstanceStaticSwitchParameterValue(MaterialInstance, AnimToTextureParamNames::UseFourInfluences, false, MaterialParameterAssociation);
+				break;
+			case EAnimToTextureNumBoneInfluences::Two:
+				UMaterialEditingLibrary::SetMaterialInstanceStaticSwitchParameterValue(MaterialInstance, AnimToTextureParamNames::UseTwoInfluences, true, MaterialParameterAssociation);
+				UMaterialEditingLibrary::SetMaterialInstanceStaticSwitchParameterValue(MaterialInstance, AnimToTextureParamNames::UseFourInfluences, false, MaterialParameterAssociation);
+				break;
+			case EAnimToTextureNumBoneInfluences::Four:
+				UMaterialEditingLibrary::SetMaterialInstanceStaticSwitchParameterValue(MaterialInstance, AnimToTextureParamNames::UseTwoInfluences, false, MaterialParameterAssociation);
+				UMaterialEditingLibrary::SetMaterialInstanceStaticSwitchParameterValue(MaterialInstance, AnimToTextureParamNames::UseFourInfluences, true, MaterialParameterAssociation);
+				break;
+		}
+
+	}
+
+	// AutoPlay
+	if (bAutoPlay && DataAsset->Animations.IsValidIndex(AnimationIndex))
+	{
+		UMaterialEditingLibrary::SetMaterialInstanceStaticSwitchParameterValue(MaterialInstance, AnimToTextureParamNames::AutoPlay, true, MaterialParameterAssociation);
+		UMaterialEditingLibrary::SetMaterialInstanceScalarParameterValue(MaterialInstance, AnimToTextureParamNames::StartFrame, DataAsset->Animations[AnimationIndex].StartFrame, MaterialParameterAssociation);
+		UMaterialEditingLibrary::SetMaterialInstanceScalarParameterValue(MaterialInstance, AnimToTextureParamNames::EndFrame, DataAsset->Animations[AnimationIndex].EndFrame, MaterialParameterAssociation);
+	}
+	
+	// NumFrames
+	UMaterialEditingLibrary::SetMaterialInstanceScalarParameterValue(MaterialInstance, AnimToTextureParamNames::NumFrames, DataAsset->NumFrames, MaterialParameterAssociation);
+
+	// SampleRate
+	UMaterialEditingLibrary::SetMaterialInstanceScalarParameterValue(MaterialInstance, AnimToTextureParamNames::SampleRate, DataAsset->SampleRate, MaterialParameterAssociation);
+
+	// Update Material
+	UMaterialEditingLibrary::UpdateMaterialInstance(MaterialInstance);
+
+	// Rebuild Material
+	UMaterialEditingLibrary::RebuildMaterialInstanceEditors(MaterialInstance->GetMaterial());
+
+	// Set Preview Mesh
+	if (DataAsset->GetStaticMesh())
+	{
+		MaterialInstance->PreviewMesh = DataAsset->GetStaticMesh();
+	}
+
+	MaterialInstance->MarkPackageDirty();
+}
+
+
+bool UAnimToTextureBPLibrary::SetLightMapIndex(UStaticMesh* StaticMesh, const int32 LODIndex, const int32 LightmapIndex, bool bGenerateLightmapUVs)
+{
+	if (!StaticMesh)
+	{
+		return false;
+	}
+
+	if (LODIndex >= 0 && !StaticMesh->IsSourceModelValid(LODIndex))
+	{
+		return false;
+	}
+
+	for (int32 Index=0; Index < LightmapIndex; Index++)
+	{
+		if (LightmapIndex > StaticMesh->GetNumUVChannels(LODIndex))
+		{
+			StaticMesh->AddUVChannel(LODIndex);
+		}
+	}
+
+	// Set Build Settings
+	FStaticMeshSourceModel& SourceModel = StaticMesh->GetSourceModel(LODIndex);
+	SourceModel.BuildSettings.bGenerateLightmapUVs = bGenerateLightmapUVs;
+	SourceModel.BuildSettings.DstLightmapIndex = LightmapIndex;
+	StaticMesh->SetLightMapCoordinateIndex(LightmapIndex);
+
+	// Build Mesh
+	StaticMesh->Build(false);
+	StaticMesh->PostEditChange();
+	StaticMesh->MarkPackageDirty();
+
+	return true;
+}
+
+
+UStaticMesh* UAnimToTextureBPLibrary::ConvertSkeletalMeshToStaticMesh(USkeletalMesh* SkeletalMesh, const FString PackageName, const int32 LODIndex)
+{
+	if (!SkeletalMesh || PackageName.IsEmpty())
+	{
+		return nullptr;
+	}
+
+	if (!FPackageName::IsValidObjectPath(PackageName))
+	{
+		return nullptr;
+	}
+
+	if (LODIndex >= 0 && !SkeletalMesh->IsValidLODIndex(LODIndex))
+	{
+		UE_LOG(LogAnimToTextureEditor, Warning, TEXT("Invalid LODIndex: %i"), LODIndex);
+		return nullptr;
+	}
+
+	// Create Temp Actor
+	check(GEditor);
+	UWorld* World = GEditor->GetEditorWorldContext().World();
+	check(World);
+	AActor* Actor = World->SpawnActor<AActor>();
+	check(Actor);
+
+	// Create Temp SkeletalMesh Component
+	USkeletalMeshComponent* MeshComponent = NewObject<USkeletalMeshComponent>(Actor);
+	MeshComponent->RegisterComponent();
+	MeshComponent->SetSkeletalMesh(SkeletalMesh);
+	TArray<UMeshComponent*> MeshComponents = { MeshComponent };
+
+	UStaticMesh* OutStaticMesh = nullptr;
+	bool bGeneratedCorrectly = true;
+
+	// Create New StaticMesh
+	if (!FPackageName::DoesPackageExist(PackageName))
+	{
+		IMeshUtilities& MeshUtilities = FModuleManager::Get().LoadModuleChecked<IMeshUtilities>("MeshUtilities");
+		OutStaticMesh = MeshUtilities.ConvertMeshesToStaticMesh(MeshComponents, FTransform::Identity, PackageName);
+	}
+	// Update Existing StaticMesh
+	else
+	{
+		// Load Existing Mesh
+		OutStaticMesh = LoadObject<UStaticMesh>(nullptr, *PackageName);
+	}
+
+	if (OutStaticMesh)
+	{
+		// Create Temp Package.
+		// because 
+		UPackage* TransientPackage = GetTransientPackage();
+
+		// Create Temp Mesh.
+		IMeshUtilities& MeshUtilities = FModuleManager::Get().LoadModuleChecked<IMeshUtilities>("MeshUtilities");
+		UStaticMesh* TempMesh = MeshUtilities.ConvertMeshesToStaticMesh(MeshComponents, FTransform::Identity, TransientPackage->GetPathName());
+
+		// make sure transactional flag is on
+		TempMesh->SetFlags(RF_Transactional);
+
+		// Copy All LODs
+		if (LODIndex < 0)
+		{
+			const int32 NumSourceModels = TempMesh->GetNumSourceModels();
+			OutStaticMesh->SetNumSourceModels(NumSourceModels);
+
+			for (int32 Index = 0; Index < NumSourceModels; ++Index)
+			{
+				// Get RawMesh
+				FRawMesh RawMesh;
+				TempMesh->GetSourceModel(Index).LoadRawMesh(RawMesh);
+
+				// Set RawMesh
+				OutStaticMesh->GetSourceModel(Index).SaveRawMesh(RawMesh);
+			};
+		}
+
+		// Copy Single LOD
+		else
+		{
+			if (LODIndex >= TempMesh->GetNumSourceModels())
+			{
+				UE_LOG(LogAnimToTextureEditor, Warning, TEXT("Invalid Source Model Index: %i"), LODIndex);
+				bGeneratedCorrectly = false;
+			}
+			else
+			{
+				OutStaticMesh->SetNumSourceModels(1);
+
+				// Get RawMesh
+				FRawMesh RawMesh;
+				TempMesh->GetSourceModel(LODIndex).LoadRawMesh(RawMesh);
+
+				// Set RawMesh
+				OutStaticMesh->GetSourceModel(0).SaveRawMesh(RawMesh);
+			}
+		}
+			
+		// Copy Materials
+		const TArray<FStaticMaterial>& Materials = TempMesh->GetStaticMaterials();
+		OutStaticMesh->SetStaticMaterials(Materials);
+
+		// Done
+		TArray<FText> OutErrors;
+		OutStaticMesh->Build(true, &OutErrors);
+		OutStaticMesh->MarkPackageDirty();
+	}
+
+	// Destroy Temp Component and Actor
+	MeshComponent->UnregisterComponent();
+	MeshComponent->DestroyComponent();
+	Actor->Destroy();
+
+	return bGeneratedCorrectly ? OutStaticMesh : nullptr;
+}
+
+
+void UAnimToTextureBPLibrary::NormalizeVertexData(
+	const TArray<FVector3f>& Deltas, const TArray<FVector3f>& Normals,
+	FVector& OutMinBBox, FVector& OutSizeBBox,
+	TArray<FVector3f>& OutNormalizedDeltas, TArray<FVector3f>& OutNormalizedNormals)
+{
+	check(Deltas.Num() == Normals.Num());
+
+	// ---------------------------------------------------------------------------
+	// Compute Bounding Box
+	//
+	OutMinBBox = { TNumericLimits<float>::Max(), TNumericLimits<float>::Max(), TNumericLimits<float>::Max() };
+	FVector3f MaxBBox = { TNumericLimits<float>::Min(), TNumericLimits<float>::Min(), TNumericLimits<float>::Min() };
+	
+	for (const FVector3f& Delta: Deltas)
+	{
+		// Find Min/Max BoundingBox
+		OutMinBBox.X = FMath::Min(Delta.X, OutMinBBox.X);
+		OutMinBBox.Y = FMath::Min(Delta.Y, OutMinBBox.Y);
+		OutMinBBox.Z = FMath::Min(Delta.Z, OutMinBBox.Z);
+
+		MaxBBox.X = FMath::Max(Delta.X, MaxBBox.X);
+		MaxBBox.Y = FMath::Max(Delta.Y, MaxBBox.Y);
+		MaxBBox.Z = FMath::Max(Delta.Z, MaxBBox.Z);
+	}
+
+	OutSizeBBox = (FVector)MaxBBox - OutMinBBox;
+
+	// ---------------------------------------------------------------------------
+	// Normalize Vertex Position Deltas
+	// Basically we want all deltas to be between [0, 1]
+	
+	// Compute Normalization Factor per-axis.
+	const FVector NormFactor = {
+		1.f / static_cast<float>(OutSizeBBox.X),
+		1.f / static_cast<float>(OutSizeBBox.Y),
+		1.f / static_cast<float>(OutSizeBBox.Z) };
+
+	OutNormalizedDeltas.SetNumUninitialized(Deltas.Num());
+	for (int32 Index = 0; Index < Deltas.Num(); ++Index)
+	{
+		OutNormalizedDeltas[Index] = (FVector3f)(((FVector)Deltas[Index] - OutMinBBox) * NormFactor);
+	}
+
+	// ---------------------------------------------------------------------------
+	// Normalize Vertex Normals
+	// And move them to [0, 1]
+	
+	OutNormalizedNormals.SetNumUninitialized(Normals.Num());
+	for (int32 Index = 0; Index < Normals.Num(); ++Index)
+	{
+		OutNormalizedNormals[Index] = (Normals[Index].GetSafeNormal() + FVector3f::OneVector) * 0.5f;
+	}
+
+}
+
+void UAnimToTextureBPLibrary::NormalizeBoneData(
+	const TArray<FVector3f>& Positions, const TArray<FVector4>& Rotations,
+	FVector& OutMinBBox, FVector& OutSizeBBox, 
+	TArray<FVector3f>& OutNormalizedPositions, TArray<FVector4>& OutNormalizedRotations)
+{
+	check(Positions.Num() == Rotations.Num());
+
+	// ---------------------------------------------------------------------------
+	// Compute Position Bounding Box
+	//
+	OutMinBBox = { TNumericLimits<float>::Max(), TNumericLimits<float>::Max(), TNumericLimits<float>::Max() };
+	FVector3f MaxBBox = { TNumericLimits<float>::Min(), TNumericLimits<float>::Min(), TNumericLimits<float>::Min() };
+
+	for (const FVector3f& Position : Positions)
+	{
+		// Find Min/Max BoundingBox
+		OutMinBBox.X = FMath::Min(Position.X, OutMinBBox.X);
+		OutMinBBox.Y = FMath::Min(Position.Y, OutMinBBox.Y);
+		OutMinBBox.Z = FMath::Min(Position.Z, OutMinBBox.Z);
+
+		MaxBBox.X = FMath::Max(Position.X, MaxBBox.X);
+		MaxBBox.Y = FMath::Max(Position.Y, MaxBBox.Y);
+		MaxBBox.Z = FMath::Max(Position.Z, MaxBBox.Z);
+	}
+
+	OutSizeBBox = (FVector)MaxBBox - OutMinBBox;
+
+	// ---------------------------------------------------------------------------
+	// Normalize Bone Position.
+	// Basically we want all positions to be between [0, 1]
+
+	// Compute Normalization Factor per-axis.
+	const FVector NormFactor = {
+		1.f / static_cast<float>(OutSizeBBox.X),
+		1.f / static_cast<float>(OutSizeBBox.Y),
+		1.f / static_cast<float>(OutSizeBBox.Z) };
+
+	OutNormalizedPositions.SetNumUninitialized(Positions.Num());
+	for (int32 Index = 0; Index < Positions.Num(); ++Index)
+	{
+		OutNormalizedPositions[Index] = FVector3f(((FVector)Positions[Index] - OutMinBBox) * NormFactor);
+	}
+
+	// ---------------------------------------------------------------------------
+	// Normalize Rotations
+	// And move them to [0, 1]
+	OutNormalizedRotations.SetNumUninitialized(Rotations.Num());
+	for (int32 Index = 0; Index < Rotations.Num(); ++Index)
+	{
+		const FVector4 Axis = Rotations[Index];
+		const float Angle = Rotations[Index].W; // Angle are returned in radians and they go from [0-pi*2]
+
+		OutNormalizedRotations[Index] = (Axis.GetSafeNormal() + FVector::OneVector) * 0.5f;
+		OutNormalizedRotations[Index].W = Angle / (PI * 2.f);
+	}
+}
+
+
+bool UAnimToTextureBPLibrary::CreateUVChannel(
+	UStaticMesh* StaticMesh, const int32 LODIndex, const int32 UVChannelIndex,
+	const int32 Height, const int32 Width)
+{
+	check(StaticMesh);
+
+	if (!StaticMesh->IsSourceModelValid(LODIndex))
+	{
+		return false;
+	}
+
+	// ----------------------------------------------------------------------------
+	// Get Mesh Description.
+	// This is needed for Inserting UVChannel
+	FMeshDescription* MeshDescription = StaticMesh->GetMeshDescription(LODIndex);
+	check(MeshDescription);
+
+	// Add New UVChannel.
+	if (UVChannelIndex == StaticMesh->GetNumUVChannels(LODIndex))
+	{
+		if (!StaticMesh->InsertUVChannel(LODIndex, UVChannelIndex))
+		{
+			UE_LOG(LogAnimToTextureEditor, Warning, TEXT("Unable to Add UVChannel"));
+			return false;
+		}
+	}
+	else if (UVChannelIndex > StaticMesh->GetNumUVChannels(LODIndex))
+	{
+		UE_LOG(LogAnimToTextureEditor, Warning, TEXT("UVChannel: %i Out of Range. Number of existing UVChannels: %i"), UVChannelIndex, StaticMesh->GetNumUVChannels(LODIndex));
+		return false;
+	}
+
+	// -----------------------------------------------------------------------------
+
+	TMap<FVertexInstanceID, FVector2D> TexCoords;
+
+	for (const FVertexInstanceID VertexInstanceID : MeshDescription->VertexInstances().GetElementIDs())
+	{
+		const FVertexID VertexID = MeshDescription->GetVertexInstanceVertex(VertexInstanceID);
+		const int32 VertexIndex = VertexID.GetValue();
+
+		float U = (0.5f / (float)Width) + (VertexIndex % Width) / (float)Width;
+		float V = (0.5f / (float)Height) + (VertexIndex / Width) / (float)Height;
+		
+		TexCoords.Add(VertexInstanceID, FVector2D(U, V));
+	}
+
+	// Set Full Precision UVs
+	SetFullPrecisionUVs(StaticMesh, LODIndex, true);
+
+	// Set UVs
+	if (StaticMesh->SetUVChannel(LODIndex, UVChannelIndex, TexCoords))
+	{
+		return true;
+	}
+	else
+	{
+		UE_LOG(LogAnimToTextureEditor, Warning, TEXT("Unable to Set UVChannel: %i. TexCoords: %i"), UVChannelIndex, TexCoords.Num());
+		return false;
+	};
+
+	return false;
+}
+
+bool UAnimToTextureBPLibrary::FindBestResolution(
+	const int32 NumFrames, const int32 NumElements,
+	int32& OutHeight, int32& OutWidth, int32& OutRowsPerFrame,
+	const int32 MaxHeight, const int32 MaxWidth, bool bEnforcePowerOfTwo)
+{
+	if (bEnforcePowerOfTwo)
+	{
+		OutWidth = 2;
+		while (OutWidth < NumElements && OutWidth < MaxWidth)
+		{
+			OutWidth *= 2;
+		}
+		OutRowsPerFrame = FMath::CeilToInt(NumElements / (float)OutWidth);
+
+		const int32 TargetHeight = NumFrames * OutRowsPerFrame;
+		OutHeight = 2;
+		while (OutHeight < TargetHeight)
+		{
+			OutHeight *= 2;
+		}
+	}
+	else
+	{
+		OutRowsPerFrame = FMath::CeilToInt(NumElements / (float)MaxWidth);
+		OutWidth = FMath::CeilToInt(NumElements / (float)OutRowsPerFrame);
+		OutHeight = NumFrames * OutRowsPerFrame;
+	}
+
+	const bool bValidResolution = OutWidth <= MaxWidth && OutHeight <= MaxHeight;
+	return bValidResolution;
+};
+
+void UAnimToTextureBPLibrary::SetFullPrecisionUVs(UStaticMesh* StaticMesh, int32 LODIndex, bool bFullPrecision)
+{
+	check(StaticMesh);
+
+	if (StaticMesh->IsSourceModelValid(LODIndex))
+	{
+		FStaticMeshSourceModel& SourceModel = StaticMesh->GetSourceModel(LODIndex);
+		SourceModel.BuildSettings.bUseFullPrecisionUVs = bFullPrecision;
+	}
+}
+
+void UAnimToTextureBPLibrary::SetBoundsExtensions(UStaticMesh* StaticMesh, const FVector& MinBBox, const FVector& SizeBBox)
+{
+	check(StaticMesh);
+
+	// Calculate MaxBBox
+	const FVector MaxBBox = SizeBBox + MinBBox;
+
+	// Reset current extension bounds
+	const FVector PositiveBoundsExtension = StaticMesh->GetPositiveBoundsExtension();
+	const FVector NegativeBoundsExtension = StaticMesh->GetNegativeBoundsExtension();
+		
+	// Get current BoundingBox including extensions
+	FBox BoundingBox = StaticMesh->GetBoundingBox();
+		
+	// Remove extensions from BoundingBox
+	BoundingBox.Max -= PositiveBoundsExtension;
+	BoundingBox.Min += NegativeBoundsExtension;
+		
+	// Calculate New BoundingBox
+	FVector NewMaxBBox(
+		FMath::Max(BoundingBox.Max.X, MaxBBox.X),
+		FMath::Max(BoundingBox.Max.Y, MaxBBox.Y),
+		FMath::Max(BoundingBox.Max.Z, MaxBBox.Z)
+	);
+		
+	FVector NewMinBBox(
+		FMath::Min(BoundingBox.Min.X, MinBBox.X),
+		FMath::Min(BoundingBox.Min.Y, MinBBox.Y),
+		FMath::Min(BoundingBox.Min.Z, MinBBox.Z)
+	);
+
+	// Calculate New Extensions
+	FVector NewPositiveBoundsExtension = NewMaxBBox - BoundingBox.Max;
+	FVector NewNegativeBoundsExtension = BoundingBox.Min - NewMinBBox;
+				
+	// Update StaticMesh
+	StaticMesh->SetPositiveBoundsExtension(NewPositiveBoundsExtension);
+	StaticMesh->SetNegativeBoundsExtension(NewNegativeBoundsExtension);
+	StaticMesh->CalculateExtendedBounds();
+}
diff --git a/Engine/Plugins/Experimental/AnimToTexture/Source/AnimToTextureEditor/Private/AnimToTextureEditorModule.cpp b/Engine/Plugins/Experimental/AnimToTexture/Source/AnimToTextureEditor/Private/AnimToTextureEditorModule.cpp
new file mode 100644
--- /dev/null
+++ b/Engine/Plugins/Experimental/AnimToTexture/Source/AnimToTextureEditor/Private/AnimToTextureEditorModule.cpp
@@ -0,0 +1,24 @@
+// Copyright Epic Games, Inc. All Rights Reserved.
+
+#include "AnimToTextureEditorModule.h"
+
+#define LOCTEXT_NAMESPACE "FAnimToTextureEditorModule"
+
+DEFINE_LOG_CATEGORY(LogAnimToTextureEditor);
+
+void FAnimToTextureEditorModule::StartupModule()
+{
+	// This code will execute after your module is loaded into memory; the exact timing is specified in the .uplugin file per-module
+	
+}
+
+void FAnimToTextureEditorModule::ShutdownModule()
+{
+	// This function may be called during shutdown to clean up your module.  For modules that support dynamic reloading,
+	// we call this function before unloading the module.
+	
+}
+
+#undef LOCTEXT_NAMESPACE
+	
+IMPLEMENT_MODULE(FAnimToTextureEditorModule, AnimToTextureEditor)
\ No newline at end of file
diff --git a/Engine/Plugins/Experimental/AnimToTexture/Source/AnimToTextureEditor/Private/AnimToTextureMeshMapping.cpp b/Engine/Plugins/Experimental/AnimToTexture/Source/AnimToTextureEditor/Private/AnimToTextureMeshMapping.cpp
new file mode 100644
--- /dev/null
+++ b/Engine/Plugins/Experimental/AnimToTexture/Source/AnimToTextureEditor/Private/AnimToTextureMeshMapping.cpp
@@ -0,0 +1,281 @@
+// Copyright Epic Games, Inc. All Rights Reserved.
+#include "AnimToTextureMeshMapping.h"
+
+namespace AnimToTexture_Private
+{
+
+FVector3f FVertexToMeshMapping::FindClosestPointToTriangle(const FVector3f& P, const FVector3f& A, const FVector3f& B, const FVector3f& C)
+{
+	const FVector3f AB = B - A;
+	const FVector3f AC = C - A;
+	const FVector3f AP = P - A;
+
+	const float D1 = FVector3f::DotProduct(AB, AP);
+	const float D2 = FVector3f::DotProduct(AC, AP);
+	if (D1 <= 0.f && D2 <= 0.f)
+	{
+		return A;
+	}
+
+	const FVector3f BP = P - B;
+	const float D3 = FVector3f::DotProduct(AB, BP);
+	const float D4 = FVector3f::DotProduct(AC, BP);
+	if (D3 >= 0.f && D4 <= D3)
+	{
+		return B;
+	}
+
+	const FVector3f CP = P - C;
+	const float D5 = FVector3f::DotProduct(AB, CP);
+	const float D6 = FVector3f::DotProduct(AC, CP);
+	if (D6 >= 0.f && D5 <= D6)
+	{
+		return C;
+	}
+
+	const float VC = D1 * D4 - D3 * D2;
+	if (VC <= 0.f && D1 >= 0.f && D3 <= 0.f)
+	{
+		const float V = D1 / (D1 - D3);
+		return A + V * AB;
+	}
+
+	const float VB = D5 * D2 - D1 * D6;
+	if (VB <= 0.f && D2 >= 0.f && D6 <= 0.f)
+	{
+		const float V = D2 / (D2 - D6);
+		return A + V * AC;
+	}
+
+	const float VA = D3 * D6 - D5 * D4;
+	if (VA <= 0.f && (D4 - D3) >= 0.f && (D5 - D6) >= 0.f)
+	{
+		const float V = (D4 - D3) / ((D4 - D3) + (D5 - D6));
+		return B + V * (C - B);
+	}
+
+	const float Denom = 1.0f / (VA + VB + VC);
+	const float V = VB * Denom;
+	const float W = VC * Denom;
+
+	return A + V * AB + W * AC;
+}
+
+FVector3f FVertexToMeshMapping::GetTriangleNormal(const FVector3f& A, const FVector3f& B, const FVector3f& C)
+{
+	const FVector3f V0 = B - A;
+	const FVector3f V1 = C - A;
+	return FVector3f::CrossProduct(V0, V1).GetSafeNormal();
+}
+
+FMatrix44f FVertexToMeshMapping::GetTriangleMatrix(const FVector3f& A, const FVector3f& B, const FVector3f& C)
+{	
+	const FVector3f V0 = (B - A).GetSafeNormal();
+	const FVector3f V1 = GetTriangleNormal(A, B, C);
+	const FVector3f V2 = FVector3f::CrossProduct(V0, V1).GetSafeNormal();
+	
+	FMatrix44f Matrix = FMatrix44f::Identity; // SetAxes doesn't set [3] elements.
+	Matrix.SetAxes(&V0, &V1, &V2, &A);
+
+	return Matrix;
+}
+
+FVector3f FVertexToMeshMapping::BarycentricCoordinates(const FVector3f& P, const FVector3f& A, const FVector3f& B, const FVector3f& C)
+{
+	const FVector3f V0 = B - A;
+	const FVector3f V1 = C - A;
+	const FVector3f V2 = P - A;
+	const float D00 = FVector3f::DotProduct(V0, V0);
+	const float D01 = FVector3f::DotProduct(V0, V1);
+	const float D11 = FVector3f::DotProduct(V1, V1);
+	const float D20 = FVector3f::DotProduct(V2, V0);
+	const float D21 = FVector3f::DotProduct(V2, V1);
+	const float Denom = 1.0f / (D00 * D11 - D01 * D01);
+	
+	const float V = (D11 * D20 - D01 * D21) * Denom;
+	const float W = (D00 * D21 - D01 * D20) * Denom;
+	const float U = 1.0f - V - W;
+
+	return FVector3f(U, V, W);
+}
+
+int32 FVertexToMeshMapping::FindClosestTriangle(const FVector3f& Point, const TArray<FVector3f>& Vertices, const TArray<FIntVector3>& Triangles,
+	FVector3f& OutClosestPoint, FVector3f& OutBarycentricCoords)
+{
+	int32 ClosestTriangle = INDEX_NONE;
+	float ClosestDistance = TNumericLimits<float>::Max();
+
+	for (int32 TriangleIndex = 0; TriangleIndex < Triangles.Num(); TriangleIndex++)
+	{
+		const FIntVector3& Triangle = Triangles[TriangleIndex];
+
+		// Check if valid indices
+		if (Vertices.IsValidIndex(Triangle.X) && 
+			Vertices.IsValidIndex(Triangle.Y) && 
+			Vertices.IsValidIndex(Triangle.Z))
+		{
+			const FVector3f& A = Vertices[Triangle.X];
+			const FVector3f& B = Vertices[Triangle.Y];
+			const FVector3f& C = Vertices[Triangle.Z];
+
+			const FVector3f ClosestPoint = FindClosestPointToTriangle(Point, A, B, C);
+			const float Distance = FVector3f::Distance(Point, ClosestPoint);
+			if (Distance < ClosestDistance)
+			{
+				OutClosestPoint = ClosestPoint;
+				ClosestTriangle = TriangleIndex;
+				ClosestDistance = Distance;
+			}
+		}
+	}
+	check(ClosestTriangle != INDEX_NONE);
+
+	// Compute Barycentric Coordinates
+	const FIntVector3& Triangle = Triangles[ClosestTriangle];
+	const FVector3f& A = Vertices[Triangle.X];
+	const FVector3f& B = Vertices[Triangle.Y];
+	const FVector3f& C = Vertices[Triangle.Z];
+	OutBarycentricCoords = BarycentricCoordinates(OutClosestPoint, A, B, C);
+
+	return ClosestTriangle;
+}
+
+void FVertexToMeshMapping::Create(
+	const UStaticMesh* StaticMesh, const int32 StaticMeshLODIndex,	
+	const USkeletalMesh* SkeletalMesh, const int32 SkeletalMeshLODIndex,
+	TArray<FVertexToMeshMapping>& OutMapping)
+{
+	check(StaticMesh);
+	check(SkeletalMesh);
+	OutMapping.Reset();
+
+	// Get StaticMesh Vertices
+	TArray<FVector3f> StaticMeshVertices;
+	TArray<FVector3f> StaticMeshNormals;
+	const int32 NumStaticMeshVertices = GetVertices(StaticMesh, StaticMeshLODIndex, StaticMeshVertices, StaticMeshNormals);
+
+	// Get SkeletalMesh Vertices
+	TArray<FVector3f> SkeletalMeshVertices;
+	const int32 NumSkeletalMeshVertices = GetVertices(SkeletalMesh, SkeletalMeshLODIndex, SkeletalMeshVertices);
+
+	// Get SkeletalMesh Triangles
+	TArray<FIntVector3> SkeletalMeshTriangles;
+	GetTriangles(SkeletalMesh, SkeletalMeshLODIndex, SkeletalMeshTriangles);
+
+	// Allocate
+	OutMapping.SetNumUninitialized(NumStaticMeshVertices);
+
+	// Fill Out Mapping
+	for (int32 VertexIndex = 0; VertexIndex < NumStaticMeshVertices; VertexIndex++)
+	{
+		const FVector3f& StaticMeshVertex = StaticMeshVertices[VertexIndex];
+
+		FVector3f ClosestPoint;
+		FVector3f BarycentricCoords;
+
+		// Find Closest Triangle (on SkeletalMesh)
+		const int32 ClosestTriangle = FindClosestTriangle(StaticMeshVertex, SkeletalMeshVertices, SkeletalMeshTriangles, ClosestPoint, BarycentricCoords);
+		check(SkeletalMeshTriangles.IsValidIndex(ClosestTriangle));
+		
+		// Closest Triangle Vertices
+		const FIntVector3& Triangle = SkeletalMeshTriangles[ClosestTriangle];
+		const FVector3f& A = SkeletalMeshVertices[Triangle.X];
+		const FVector3f& B = SkeletalMeshVertices[Triangle.Y];
+		const FVector3f& C = SkeletalMeshVertices[Triangle.Z];
+		
+		// Compute Triangle Matrix
+		const FMatrix44f Matrix = GetTriangleMatrix(A, B, C);
+		
+		// Store Mapping 
+		OutMapping[VertexIndex].Triangle = Triangle;
+		OutMapping[VertexIndex].BarycentricCoords = BarycentricCoords;
+		OutMapping[VertexIndex].InvMatrix = Matrix.Inverse();
+	}
+}
+
+
+void FVertexToMeshMapping::InterpolateVertexSkinWeights(const VertexSkinWeightMax& A, const VertexSkinWeightMax& B, const VertexSkinWeightMax& C,
+	VertexSkinWeightMax& OutWeights) const
+{	
+	// NOTE: this could use NumBones and remove the 256 limit.
+	TStaticArray<float, 256> DenseA(InPlace, 0.f);
+	TStaticArray<float, 256> DenseB(InPlace, 0.f);
+	TStaticArray<float, 256> DenseC(InPlace, 0.f);
+
+	// Convert to Dense Arrays	
+	for (int32 Index = 0; Index < MAX_TOTAL_INFLUENCES; Index++)
+	{
+		if (A.BoneWeights[Index])
+		{
+			DenseA[A.MeshBoneIndices[Index]] = (float)A.BoneWeights[Index] / 255.f;
+		}
+		if (B.BoneWeights[Index])
+		{
+			DenseB[B.MeshBoneIndices[Index]] = (float)B.BoneWeights[Index] / 255.f;
+		}
+		if (C.BoneWeights[Index])
+		{
+			DenseC[C.MeshBoneIndices[Index]] = (float)C.BoneWeights[Index] / 255.f;
+		}
+	}
+
+	// Blend Dense Arrays
+	TArray<TPair<uint8, uint16>> SortedVertexSkinWeights;
+	SortedVertexSkinWeights.SetNumUninitialized(256);
+
+	for (int32 Index = 0; Index < 256; Index++)
+	{
+		const float BoneWeightA = DenseA[Index] * BarycentricCoords.X;
+		const float BoneWeightB = DenseB[Index] * BarycentricCoords.Y;
+		const float BoneWeightC = DenseC[Index] * BarycentricCoords.Z;
+
+		const uint8 BoneWeight = (uint8)FMath::RoundToInt((BoneWeightA + BoneWeightB + BoneWeightC) * 255.f);
+		SortedVertexSkinWeights[Index] = TPair<uint8, uint16>(BoneWeight, Index);
+	}
+
+	// Sort Weights and return Sparse
+	SortedVertexSkinWeights.Sort();
+	for (int32 Index = 0; Index < MAX_TOTAL_INFLUENCES; Index++)
+	{
+		OutWeights.BoneWeights[Index]     = SortedVertexSkinWeights[255 - Index].Key;
+		OutWeights.MeshBoneIndices[Index] = SortedVertexSkinWeights[255 - Index].Value;
+	}
+}
+
+void FVertexToMeshMapping::InterpolateSkinWeights(const TArray<FVertexToMeshMapping>& Mapping, const TArray<VertexSkinWeightMax>& SkinWeights,
+	TArray<VertexSkinWeightMax>& OutSkinWeights)
+{
+	// Allocate StaticMesh Weights
+	OutSkinWeights.SetNumUninitialized(Mapping.Num());
+
+	for (int32 VertexIndex = 0; VertexIndex < Mapping.Num(); VertexIndex++)
+	{
+		const FVertexToMeshMapping& VertexMapping = Mapping[VertexIndex];
+		const FIntVector3& Triangle = VertexMapping.Triangle;
+
+		if (SkinWeights.IsValidIndex(Triangle.X) &&
+			SkinWeights.IsValidIndex(Triangle.Y) &&
+			SkinWeights.IsValidIndex(Triangle.Z))
+		{
+			const VertexSkinWeightMax& A = SkinWeights[Triangle.X];
+			const VertexSkinWeightMax& B = SkinWeights[Triangle.Y];
+			const VertexSkinWeightMax& C = SkinWeights[Triangle.Z];
+
+			Mapping[VertexIndex].InterpolateVertexSkinWeights(A, B, C, OutSkinWeights[VertexIndex]);
+		}
+	}
+}
+
+FVector3f FVertexToMeshMapping::TransformPosition(const FVector3f& P, const FVector3f& A, const FVector3f& B, const FVector3f& C) const
+{	
+	const FMatrix44f Matrix = GetTriangleMatrix(A, B, C);	
+	return Matrix.TransformPosition(InvMatrix.TransformPosition(P));
+}
+
+FVector3f FVertexToMeshMapping::TransformVector(const FVector3f& V, const FVector3f& A, const FVector3f& B, const FVector3f& C) const
+{
+	const FMatrix44f Matrix = GetTriangleMatrix(A, B, C);
+	return Matrix.TransformVector(InvMatrix.TransformVector(V));
+}
+
+} // end namespace AnimToTexture_Private
diff --git a/Engine/Plugins/Experimental/AnimToTexture/Source/AnimToTextureEditor/Private/AnimToTextureSkeletalMesh.cpp b/Engine/Plugins/Experimental/AnimToTexture/Source/AnimToTextureEditor/Private/AnimToTextureSkeletalMesh.cpp
new file mode 100644
--- /dev/null
+++ b/Engine/Plugins/Experimental/AnimToTexture/Source/AnimToTextureEditor/Private/AnimToTextureSkeletalMesh.cpp
@@ -0,0 +1,330 @@
+// Copyright Epic Games, Inc. All Rights Reserved.
+#include "AnimToTextureSkeletalMesh.h"
+#include "AnimToTextureEditorModule.h"
+#include "Rendering/SkeletalMeshRenderData.h"
+#include "Engine/StaticMesh.h"
+#include "Engine/SkeletalMesh.h"
+#include "Components/SkeletalMeshComponent.h"
+#include "MeshDescription.h"
+#include "StaticMeshAttributes.h"
+
+namespace AnimToTexture_Private
+{
+
+int32 GetVertices(const UStaticMesh* StaticMesh, const int32 LODIndex,
+	TArray<FVector3f>& OutPositions, TArray<FVector3f>& OutNormals)
+{
+	check(StaticMesh);
+	OutPositions.Reset();
+	OutNormals.Reset();
+
+	if (!StaticMesh->IsMeshDescriptionValid(LODIndex))
+	{
+		return INDEX_NONE;
+	}
+
+	// Get Mesh Description
+	const FMeshDescription* MeshDescription = StaticMesh->GetMeshDescription(LODIndex);
+	check(MeshDescription);
+	
+	// Get MeshDescription Vertices
+	const FVertexArray& Vertices = MeshDescription->Vertices();
+	const int32 NumVertices = Vertices.Num();
+	OutPositions.SetNum(NumVertices);
+	OutNormals.SetNumZeroed(NumVertices);
+	
+	for (const FVertexID& VertexID : Vertices.GetElementIDs())
+	{
+		const int32 VertexIndex = VertexID.GetValue();
+
+		// Get Vertex Position
+		OutPositions[VertexIndex] = MeshDescription->GetVertexPosition(VertexID);
+
+		// Get Avg Vertex Normals
+		const TArrayView<const FVertexInstanceID> VertexInstances = MeshDescription->GetVertexVertexInstanceIDs(VertexID);
+		for (const FVertexInstanceID VertexInstanceID : VertexInstances)
+		{
+			const FVector3f VertexInstanceNormal = MeshDescription->VertexInstanceAttributes().GetAttribute<FVector3f>(VertexInstanceID, MeshAttribute::VertexInstance::Normal);
+			OutNormals[VertexIndex] += VertexInstanceNormal / VertexInstances.Num();
+		}
+	}
+	
+	return NumVertices;
+}
+
+int32 GetVertices(const USkeletalMesh* SkeletalMesh, const int32 LODIndex, 
+	TArray<FVector3f>& OutPositions)
+{
+	check(SkeletalMesh);
+	OutPositions.Reset();
+	
+	const FSkeletalMeshRenderData* RenderData = SkeletalMesh->GetResourceForRendering();
+	check(RenderData);
+
+	if (!RenderData->LODRenderData.IsValidIndex(LODIndex))
+	{
+		return INDEX_NONE;
+	}
+
+	// Get LOD Data
+	const FSkeletalMeshLODRenderData& LODRenderData = RenderData->LODRenderData[LODIndex];
+
+	// Get Total Num of Vertices (for all sections)
+	const int32 NumVertices = LODRenderData.GetNumVertices();
+	OutPositions.SetNumUninitialized(NumVertices);
+	
+	for (int32 VertexIndex = 0; VertexIndex < NumVertices; ++VertexIndex)
+	{
+		OutPositions[VertexIndex] = LODRenderData.StaticVertexBuffers.PositionVertexBuffer.VertexPosition(VertexIndex);
+	};
+
+	return NumVertices;
+}
+
+int32 GetTriangles(const USkeletalMesh* SkeletalMesh, const int32 LODIndex,
+	TArray<FIntVector3>& OutTriangles)
+{
+	check(SkeletalMesh);
+	OutTriangles.Reset();
+
+	const FSkeletalMeshRenderData* RenderData = SkeletalMesh->GetResourceForRendering();
+	check(RenderData);
+
+	if (!RenderData->LODRenderData.IsValidIndex(LODIndex))
+	{
+		return INDEX_NONE;
+	}
+
+	// Get LOD Data
+	const FSkeletalMeshLODRenderData& LODRenderData = RenderData->LODRenderData[LODIndex];
+
+	// Get Indices
+	const FMultiSizeIndexContainer& IndexBuffer = LODRenderData.MultiSizeIndexContainer;
+	TArray<uint32> IndexArray;
+	IndexBuffer.GetIndexBuffer(IndexArray);
+
+	// Allocate Triangles
+	const uint32 NumTriangles = IndexArray.Num() / 3;
+	OutTriangles.SetNumUninitialized(NumTriangles);
+
+	for (uint32 TriangleIndex = 0; TriangleIndex < NumTriangles; TriangleIndex++)
+	{
+		OutTriangles[TriangleIndex].X = IndexArray[TriangleIndex * 3];
+		OutTriangles[TriangleIndex].Y = IndexArray[TriangleIndex * 3 + 1];
+		OutTriangles[TriangleIndex].Z = IndexArray[TriangleIndex * 3 + 2];
+	}
+
+	return NumTriangles;
+};
+
+int32 GetNumBones(const USkeletalMesh* SkeletalMesh)
+{
+	check(SkeletalMesh);
+	
+	const FReferenceSkeleton& RefSkeleton = SkeletalMesh->GetRefSkeleton();
+	return RefSkeleton.GetRawBoneNum();
+}
+
+/**
+* Copied from FAnimationRuntime::GetComponentSpaceTransform
+*/
+void GetRefBoneTransforms(const USkeletalMesh* SkeletalMesh, TArray<FTransform>& Transforms)
+{
+	check(SkeletalMesh);
+	Transforms.Reset();
+
+	// Get Reference Skeleton
+	const FReferenceSkeleton& RefSkeleton = SkeletalMesh->GetRefSkeleton();
+	const TArray<FTransform>& BoneSpaceTransforms = RefSkeleton.GetRawRefBonePose(); // Get only raw bones (no virtual)
+	
+	const int32 NumTransforms = BoneSpaceTransforms.Num();
+	Transforms.SetNumUninitialized(NumTransforms);
+
+	for (int32 BoneIndex = 0; BoneIndex < NumTransforms; ++BoneIndex)
+	{
+		// initialize to identity since some of them don't have tracks
+		int32 IterBoneIndex = BoneIndex;
+		FTransform ComponentSpaceTransform = BoneSpaceTransforms[BoneIndex];
+
+		do
+		{
+			//const int32 ParentIndex = RefSkeleton.GetParentIndex(IterBoneIndex);
+			const int32 ParentIndex = RefSkeleton.GetRawParentIndex(IterBoneIndex); // Get only raw bones (no virtual)
+			if (ParentIndex != INDEX_NONE)
+			{
+				ComponentSpaceTransform *= BoneSpaceTransforms[ParentIndex];
+			}
+
+			IterBoneIndex = ParentIndex;
+		} while (RefSkeleton.IsValidIndex(IterBoneIndex));
+
+		// 
+		Transforms[BoneIndex] = ComponentSpaceTransform;
+	}
+}
+
+bool HasBone(const USkeletalMesh* SkeletalMesh, const FName& Bone)
+{
+	check(SkeletalMesh);
+	
+	const FReferenceSkeleton& RefSkeleton = SkeletalMesh->GetRefSkeleton();
+
+	for (const FMeshBoneInfo& BoneInfo: RefSkeleton.GetRawRefBoneInfo())
+	{
+		if (BoneInfo.Name == Bone)
+		{
+			return true;
+		}
+	}
+	return false;
+}
+
+
+void GetBoneNames(const USkeletalMesh* SkeletalMesh, TArray<FName>& OutNames)
+{
+	check(SkeletalMesh);
+	OutNames.Reset();
+	
+	const FReferenceSkeleton& RefSkeleton = SkeletalMesh->GetRefSkeleton();
+	
+	const int32 NumRawBones = RefSkeleton.GetRawBoneNum();
+	OutNames.SetNumUninitialized(NumRawBones);
+
+	for (int32 BoneIndex = 0; BoneIndex < NumRawBones; ++BoneIndex)
+	{
+		OutNames[BoneIndex] = RefSkeleton.GetBoneName(BoneIndex);
+	}
+};
+
+void GetSkinWeights(const USkeletalMesh* SkeletalMesh, const int32 LODIndex, 
+	TArray<VertexSkinWeightMax>& OutSkinWeights)
+{
+	check(SkeletalMesh);
+	OutSkinWeights.Reset();
+
+	// Get Render Data
+	const FSkeletalMeshRenderData* RenderData = SkeletalMesh->GetResourceForRendering();
+	check(RenderData);
+
+	if (!RenderData->LODRenderData.IsValidIndex(LODIndex))
+	{
+		UE_LOG(LogAnimToTextureEditor, Warning, TEXT("Invalid LODIndex: %i"), LODIndex)
+		return;
+	}
+
+	// Get LOD Data
+	const FSkeletalMeshLODRenderData& LODRenderData = RenderData->LODRenderData[LODIndex];
+
+	// Get Weights
+	const FSkinWeightVertexBuffer* SkinWeightVertexBuffer = LODRenderData.GetSkinWeightVertexBuffer();
+	check(SkinWeightVertexBuffer);
+
+	// Get Weights from Buffer.
+	// this is size of number of vertices.
+	TArray<FSkinWeightInfo> SkinWeightsInfo;
+	SkinWeightVertexBuffer->GetSkinWeights(SkinWeightsInfo);
+
+	// Allocated SkinWeightData
+	OutSkinWeights.SetNumUninitialized(SkinWeightsInfo.Num());
+
+	// Loop thru vertices
+	for (int32 VertexIndex = 0; VertexIndex < SkinWeightsInfo.Num(); VertexIndex++)
+	{
+		// Find Section From Global Vertex Index
+		// NOTE: BoneMap is stored by Section.
+		int32 OutSectionIndex;
+		int32 OutSectionVertexIndex;
+		LODRenderData.GetSectionFromVertexIndex(VertexIndex, OutSectionIndex, OutSectionVertexIndex);
+
+		// Get Section for Vertex.
+		const FSkelMeshRenderSection& RenderSection = LODRenderData.RenderSections[OutSectionIndex];
+		
+		// Get Vertex Weights
+		const FSkinWeightInfo& SkinWeightInfo = SkinWeightsInfo[VertexIndex];
+
+		// Store Weights
+		for (int32 Index = 0; Index < MAX_TOTAL_INFLUENCES; Index++)
+		{
+			const uint8& BoneWeight = SkinWeightInfo.InfluenceWeights[Index];
+			const uint16& BoneIndex = SkinWeightInfo.InfluenceBones[Index];
+			const uint16& MeshBoneIndex = RenderSection.BoneMap[BoneIndex];
+
+			OutSkinWeights[VertexIndex].BoneWeights[Index] = BoneWeight;
+			OutSkinWeights[VertexIndex].MeshBoneIndices[Index] = MeshBoneIndex;
+		}
+	}
+}
+
+void ReduceSkinWeights(const TArray<VertexSkinWeightMax>& InSkinWeights, 
+	TArray<VertexSkinWeightFour>& OutSkinWeights)
+{
+	OutSkinWeights.SetNumUninitialized(InSkinWeights.Num());
+
+	for (int32 VertexIndex = 0; VertexIndex < InSkinWeights.Num(); VertexIndex++)
+	{
+		const VertexSkinWeightMax& InVertexSkinWeight = InSkinWeights[VertexIndex];
+		VertexSkinWeightFour& OutVertexSkinWeight = OutSkinWeights[VertexIndex];
+
+		float TotalBoneWeight = 0.f;
+		for (int32 Index = 0; Index < 4; Index++)
+		{
+			TotalBoneWeight += (float)InVertexSkinWeight.BoneWeights[Index] / 255.f;
+		}
+
+		// We assume Weights are sorted
+		for (int32 Index = 0; Index < 4; Index++)
+		{
+			OutVertexSkinWeight.BoneWeights[Index] = (uint8)FMath::RoundToInt((float)InVertexSkinWeight.BoneWeights[Index] / TotalBoneWeight);
+			OutVertexSkinWeight.MeshBoneIndices[Index] = InVertexSkinWeight.MeshBoneIndices[Index];
+		}
+	}
+}
+
+void GetSkinnedVertices(const USkeletalMeshComponent* SkeletalMeshComponent, const int32 LODIndex,
+	TArray<FVector3f>& OutPositions)
+{
+	check(SkeletalMeshComponent);
+	OutPositions.Reset();
+
+	// Get SkeletalMesh
+	const USkeletalMesh* SkeletalMesh = SkeletalMeshComponent->GetSkeletalMeshAsset();
+	check(SkeletalMesh);
+
+	// Get Matrices
+	TArray<FMatrix44f> RefToLocals;
+	SkeletalMeshComponent->CacheRefToLocalMatrices(RefToLocals);
+
+	// Get Ref-Pose Vertices
+	TArray<FVector3f> Vertices;
+	const int32 NumVertices = GetVertices(SkeletalMesh, LODIndex, Vertices);
+	OutPositions.SetNumUninitialized(NumVertices);
+
+	// TODO: Add Morph Deltas to Vertices.
+
+	// Get Weights
+	TArray<VertexSkinWeightMax> SkinWeights;
+	GetSkinWeights(SkeletalMesh, LODIndex, SkinWeights);
+
+	for (int32 VertexIndex = 0; VertexIndex < NumVertices; VertexIndex++)
+	{
+		const FVector3f& Vertex = Vertices[VertexIndex];
+		const VertexSkinWeightMax& Weights = SkinWeights[VertexIndex];
+
+		FVector4f SkinnedVertex(0);
+		for (int32 Index = 0; Index < MAX_TOTAL_INFLUENCES; Index++)
+		{
+			const uint8& BoneWeight = Weights.BoneWeights[Index];
+			const uint16& MeshBoneIndex = Weights.MeshBoneIndices[Index];
+
+			// Get Matrix
+			const FMatrix44f& RefToLocal = RefToLocals[MeshBoneIndex];
+
+			const float Weight = (float)BoneWeight / 255.f;
+			SkinnedVertex += RefToLocal.TransformPosition(Vertex) * Weight;
+		}
+
+		OutPositions[VertexIndex] = SkinnedVertex;
+	};
+};
+
+} // end namespace AnimToTexture_Private
diff --git a/Engine/Plugins/Experimental/AnimToTexture/Source/AnimToTextureEditor/Private/AnimToTextureUtils.cpp b/Engine/Plugins/Experimental/AnimToTexture/Source/AnimToTextureEditor/Private/AnimToTextureUtils.cpp
new file mode 100644
--- /dev/null
+++ b/Engine/Plugins/Experimental/AnimToTexture/Source/AnimToTextureEditor/Private/AnimToTextureUtils.cpp
@@ -0,0 +1,70 @@
+// Copyright Epic Games, Inc. All Rights Reserved.
+
+#include "AnimToTextureUtils.h"
+
+namespace AnimToTexture_Private
+{
+
+void DecomposeTransformation(const FTransform& Transform, 
+	FVector3f& OutTranslation, FVector4& OutRotation)
+{
+	// Get Translation
+	OutTranslation = (FVector3f)Transform.GetTranslation();
+
+	// Get Rotation 
+	const FQuat Quat = Transform.GetRotation();
+
+	FVector Axis;
+	float Angle;
+	Quat.ToAxisAndAngle(Axis, Angle);
+
+	OutRotation = FVector4(Axis, Angle);
+}
+
+void DecomposeTransformations(const TArray<FTransform>& Transforms, 
+	TArray<FVector3f>& OutTranslations, TArray<FVector4>& OutRotations)
+{
+	const int32 NumTransforms = Transforms.Num();
+	OutTranslations.SetNumUninitialized(NumTransforms);
+	OutRotations.SetNumUninitialized(NumTransforms);
+
+	for (int32 Index = 0; Index < NumTransforms; ++Index)
+	{
+		DecomposeTransformation(Transforms[Index], OutTranslations[Index], OutRotations[Index]);
+	}
+};
+
+bool WriteSkinWeightsToTexture(const TArray<VertexSkinWeightFour>& SkinWeights,
+	const int32 RowsPerFrame, const int32 Height, const int32 Width, UTexture2D* Texture)
+{
+	check(Texture);
+
+	const int32 NumVertices = SkinWeights.Num();
+
+	// Allocate PixelData.
+	TArray<FLowPrecision::ColorType> Pixels;
+	Pixels.Init(FLowPrecision::ColorType::Black, Height * Width);
+
+	for (int32 VertexIndex = 0; VertexIndex < NumVertices; ++VertexIndex)
+	{
+		const VertexSkinWeightFour& VertexSkinWeight = SkinWeights[VertexIndex];
+
+		// Write Influence
+		// NOTE: we are assuming the bone index is under < 256
+		Pixels[VertexIndex].R = (uint8)VertexSkinWeight.MeshBoneIndices[0];
+		Pixels[VertexIndex].G = (uint8)VertexSkinWeight.MeshBoneIndices[1];
+		Pixels[VertexIndex].B = (uint8)VertexSkinWeight.MeshBoneIndices[2];
+		Pixels[VertexIndex].A = (uint8)VertexSkinWeight.MeshBoneIndices[3];
+			
+		// Write Weight
+		Pixels[Width * RowsPerFrame + VertexIndex].R = VertexSkinWeight.BoneWeights[0];
+		Pixels[Width * RowsPerFrame + VertexIndex].G = VertexSkinWeight.BoneWeights[1];
+		Pixels[Width * RowsPerFrame + VertexIndex].B = VertexSkinWeight.BoneWeights[2];
+		Pixels[Width * RowsPerFrame + VertexIndex].A = VertexSkinWeight.BoneWeights[3];
+	};
+
+	// Write to Texture
+	return WriteToTexture<FLowPrecision>(Texture, Height, Width, Pixels);
+}
+
+} // end namespace AnimToTexture_Private
\ No newline at end of file
diff --git a/Engine/Plugins/Experimental/AnimToTexture/Source/AnimToTextureEditor/Public/AnimToTextureBPLibrary.h b/Engine/Plugins/Experimental/AnimToTexture/Source/AnimToTextureEditor/Public/AnimToTextureBPLibrary.h
new file mode 100644
--- /dev/null
+++ b/Engine/Plugins/Experimental/AnimToTexture/Source/AnimToTextureEditor/Public/AnimToTextureBPLibrary.h
@@ -0,0 +1,105 @@
+// Copyright Epic Games, Inc. All Rights Reserved.
+
+#pragma once
+
+#include "AnimToTextureDataAsset.h"
+#include "AnimToTextureMeshMapping.h"
+#include "Kismet/BlueprintFunctionLibrary.h"
+#include "Materials/MaterialLayersFunctions.h"
+
+#include "AnimToTextureBPLibrary.generated.h"
+
+class UStaticMesh;
+class USkeletalMesh;
+
+/*
+* TODO: 
+*	- Right now it is saving data per-vertex instead of per-point.
+*     This will require more pixels if the mesh has lots of material slots or uv-shells.
+*     FStaticMeshOperations::FindOverlappingCorners(FOverlappingCorners& OutOverlappingCorners, const FMeshDescription& MeshDescription, float ComparisonThreshold)
+*/
+UCLASS()
+class UAnimToTextureBPLibrary : public UBlueprintFunctionLibrary
+{
+	GENERATED_UCLASS_BODY()
+
+#if WITH_EDITOR
+
+	/**
+	* Bakes Animation Data into Textures.
+	*/
+	UFUNCTION(BlueprintCallable, meta = (Category = "AnimToTexture"))
+	static void AnimationToTexture(UAnimToTextureDataAsset* DataAsset, const FTransform RootTransform, bool& bSuccess);
+
+	/** 
+	* Utility for converting SkeletalMesh into a StaticMesh
+	*/
+	UFUNCTION(BlueprintCallable, Category = "AnimToTexture")
+	static UStaticMesh* ConvertSkeletalMeshToStaticMesh(USkeletalMesh* SkeletalMesh, const FString PackageName, const int32 LODIndex = -1);
+
+	/**
+	* Utility for setting a StaticMesh LightMapIndex.
+	*/
+	UFUNCTION(BlueprintCallable, Category = "AnimToTexture")
+	static bool SetLightMapIndex(UStaticMesh* StaticMesh, const int32 LODIndex, const int32 LightmapIndex=1, bool bGenerateLightmapUVs=true);
+
+	/**
+	* Updates a material's parameters to match those of an AnimToTexture DataAsset
+	* @param AnimationIndex Index of the animation to play.
+	* @param NumBoneInfluences Number of skinning influences. More influences will produce smoother deformation but slower performance. This is only used in Bone Mode.
+	*/
+	UFUNCTION(BlueprintCallable, meta = (Category = "AnimToTexture"))
+	static void UpdateMaterialInstanceFromDataAsset(UAnimToTextureDataAsset* DataAsset, class UMaterialInstanceConstant* MaterialInstance, 
+		const bool bAutoPlay=true, const int32 AnimationIndex=0, 
+		const EAnimToTextureNumBoneInfluences NumBoneInfluences = EAnimToTextureNumBoneInfluences::One,
+		const EMaterialParameterAssociation MaterialParameterAssociation = EMaterialParameterAssociation::LayerParameter);
+
+#endif // WITH_EDITOR
+
+private:
+
+	// Get Vertex and Normals from Current Pose
+	// The VertexDelta is returned from the RefPose
+	static void GetVertexDeltasAndNormals(const USkeletalMeshComponent* SkeletalMeshComponent, const int32 LODIndex, 
+		const TArray<FVector3f>& Vertices, const TArray<FVector3f>& Normals, const TArray<AnimToTexture_Private::FVertexToMeshMapping>& Mapping,
+		const FTransform RootTransform,
+		TArray<FVector3f>& OutVertexDeltas, TArray<FVector3f>& OutVertexNormals);
+
+	// Gets RefPose Bone Position and Rotations.
+	static int32 GetRefBonePositionsAndRotations(const USkeletalMeshComponent* SkeletalMeshComponent, 
+		TArray<FVector3f>& OutBoneRefPositions, TArray<FVector4>& OutBoneRefRotations);
+
+	// Gets Bone Position and Rotations for Current Pose.	
+	// The BonePosition is returned relative to the RefPose
+	static int32 GetBonePositionsAndRotations(const USkeletalMeshComponent* SkeletalMeshComponent, const TArray<FVector3f>& BoneRefPositions,
+		TArray<FVector3f>& BonePositions, TArray<FVector4>& BoneRotations);
+
+	// Normalizes Deltas and Normals between [0-1] with Bounding Box
+	static void NormalizeVertexData(
+		const TArray<FVector3f>& Deltas, const TArray<FVector3f>& Normals,
+		FVector& OutMinBBox, FVector& OutSizeBBox,
+		TArray<FVector3f>& OutNormalizedDeltas, TArray<FVector3f>& OutNormalizedNormals);
+
+	// Normalizes Positions and Rotations between [0-1] with Bounding Box
+	static void NormalizeBoneData(
+		const TArray<FVector3f>& Positions, const TArray<FVector4>& Rotations,
+		FVector& OutMinBBox, FVector& OutSizeBBox,
+		TArray<FVector3f>& OutNormalizedPositions, TArray<FVector4>& OutNormalizedRotations);
+
+	/* Returns best resolution for the given data. 
+	*  Returns false if data doesnt fit in the the max range */
+	static bool FindBestResolution(const int32 NumFrames, const int32 NumElements,
+								   int32& OutHeight, int32& OutWidth, int32& OutRowsPerFrame,
+								   const int32 MaxHeight = 4096, const int32 MaxWidth = 4096, bool bEnforcePowerOfTwo = false);
+
+	/* Sets Static Mesh FullPrecisionUVs Property*/
+	static void SetFullPrecisionUVs(UStaticMesh* StaticMesh, const int32 LODIndex, bool bFullPrecision=true);
+
+	/* Sets Static Mesh Bound Extensions */
+	static void SetBoundsExtensions(UStaticMesh* StaticMesh, const FVector& MinBBox, const FVector& SizeBBox);
+
+	/* Creates UV Coord with vertices */
+	static bool CreateUVChannel(UStaticMesh* StaticMesh, const int32 LODIndex, const int32 UVChannelIndex,
+		const int32 Height, const int32 Width);
+
+};
diff --git a/Engine/Plugins/Experimental/AnimToTexture/Source/AnimToTextureEditor/Public/AnimToTextureEditorModule.h b/Engine/Plugins/Experimental/AnimToTexture/Source/AnimToTextureEditor/Public/AnimToTextureEditorModule.h
new file mode 100644
--- /dev/null
+++ b/Engine/Plugins/Experimental/AnimToTexture/Source/AnimToTextureEditor/Public/AnimToTextureEditorModule.h
@@ -0,0 +1,16 @@
+// Copyright Epic Games, Inc. All Rights Reserved.
+
+#pragma once
+
+#include "Modules/ModuleManager.h"
+
+DECLARE_LOG_CATEGORY_EXTERN(LogAnimToTextureEditor, Log, All);
+
+class FAnimToTextureEditorModule : public IModuleInterface
+{
+public:
+
+	/** IModuleInterface implementation */
+	virtual void StartupModule() override;
+	virtual void ShutdownModule() override;
+};
diff --git a/Engine/Plugins/Experimental/AnimToTexture/Source/AnimToTextureEditor/Public/AnimToTextureMeshMapping.h b/Engine/Plugins/Experimental/AnimToTexture/Source/AnimToTextureEditor/Public/AnimToTextureMeshMapping.h
new file mode 100644
--- /dev/null
+++ b/Engine/Plugins/Experimental/AnimToTexture/Source/AnimToTextureEditor/Public/AnimToTextureMeshMapping.h
@@ -0,0 +1,57 @@
+// Copyright Epic Games, Inc. All Rights Reserved.
+
+#pragma once
+
+#include "AnimToTextureSkeletalMesh.h"
+#include "CoreMinimal.h"
+
+namespace AnimToTexture_Private
+{
+
+// Static to Skeletal Mapping
+struct FVertexToMeshMapping
+{
+	// Closest SkeletalMesh Triangle
+	FIntVector3 Triangle;
+	// SkeletalMesh Triangle Barycentric Coordinates
+	FVector3f  BarycentricCoords;
+	// SkeletalMesh Triangle Inverted Matrix
+	FMatrix44f InvMatrix;
+	
+	/* Transforms Point with given Triangle */
+	FVector3f TransformPosition(const FVector3f& Point, const FVector3f& PointA, const FVector3f& PointB, const FVector3f& PointC) const;
+	
+	/* Transforms Vector with given Triangle */
+	FVector3f TransformVector(const FVector3f& Vector, const FVector3f& PointA, const FVector3f& PointB, const FVector3f& PointC) const;
+	
+	/* Interpolate VertexSkinWeights with Barycentric Coordinates */
+	void InterpolateVertexSkinWeights(const VertexSkinWeightMax& A, const VertexSkinWeightMax& B, const VertexSkinWeightMax& C,
+		VertexSkinWeightMax& OutWeights) const;
+	
+	/* Computes closest point to triangle*/
+	static FVector3f FindClosestPointToTriangle(const FVector3f& Point, const FVector3f& PointA, const FVector3f& PointB, const FVector3f& PointC);
+
+	/* Computes Barycentric coordinates from point to triangle */
+	static FVector3f BarycentricCoordinates(const FVector3f& Point, const FVector3f& PointA, const FVector3f& PointB, const FVector3f& PointC);
+
+	/* Computes Triangle Normal */
+	static FVector3f GetTriangleNormal(const FVector3f& PointA, const FVector3f& PointB, const FVector3f& PointC);
+
+	/* Computes Triangle Matrix */
+	static FMatrix44f GetTriangleMatrix(const FVector3f& PointA, const FVector3f& PointB, const FVector3f& PointC);
+
+	/* Finds closest triangle, point on triangle and barycentric coordinates */
+	static int32 FindClosestTriangle(const FVector3f& Point, const TArray<FVector3f>& Vertices, const TArray<FIntVector3>& Triangles,
+		FVector3f& OutClosestPoint, FVector3f& OutBarycentricCoords);
+
+	// Creates Mapping between StaticMesh and SkeletalMesh
+	static void Create(const UStaticMesh* StaticMesh, const int32 StaticMeshLODIndex,
+		const USkeletalMesh* SkeletalMesh, const int32 SkeletalMeshLODIndex,
+		TArray<FVertexToMeshMapping>& OutMapping);
+
+	/* Interpolate SkinWeights with VertexToMeshMapping */
+	static void InterpolateSkinWeights(const TArray<FVertexToMeshMapping>& Mapping, const TArray<VertexSkinWeightMax>& SkinWeights,
+		TArray<VertexSkinWeightMax>& OutSkinWeights);
+};
+
+} // end namespace AnimToTexture_Private
diff --git a/Engine/Plugins/Experimental/AnimToTexture/Source/AnimToTextureEditor/Public/AnimToTextureSkeletalMesh.h b/Engine/Plugins/Experimental/AnimToTexture/Source/AnimToTextureEditor/Public/AnimToTextureSkeletalMesh.h
new file mode 100644
--- /dev/null
+++ b/Engine/Plugins/Experimental/AnimToTexture/Source/AnimToTextureEditor/Public/AnimToTextureSkeletalMesh.h
@@ -0,0 +1,61 @@
+// Copyright Epic Games, Inc. All Rights Reserved.
+
+#pragma once
+
+#include "AnimToTextureDataAsset.h"
+#include "CoreMinimal.h"
+#include "Containers/StaticArray.h"
+#include "GPUSkinPublicDefs.h" 
+
+namespace AnimToTexture_Private
+{
+
+template <uint16 NumInfluences>
+struct TVertexSkinWeight
+{
+	TStaticArray<uint16, NumInfluences> MeshBoneIndices;
+	TStaticArray<uint8, NumInfluences>  BoneWeights;
+};
+
+using VertexSkinWeightMax  = TVertexSkinWeight<MAX_TOTAL_INFLUENCES>;
+using VertexSkinWeightFour = TVertexSkinWeight<4>;
+
+/* Returns StaticMesh Vertex Positions */
+int32 GetVertices(const UStaticMesh* StaticMesh, const int32 LODIndex,
+	TArray<FVector3f>& OutPositions, TArray<FVector3f>& OutNormals);
+
+/* Returns RefPose Vertex Positions */
+int32 GetVertices(const USkeletalMesh* SkeletalMesh, const int32 LODIndex,
+	TArray<FVector3f>& OutPositions);
+
+/* Returns Triangles vertex indices */
+int32 GetTriangles(const USkeletalMesh* SkeletalMesh, const int32 LODIndex,
+	TArray<FIntVector3>& OutTriangles);
+
+/* Computes CPUSkinning at Pose */
+void GetSkinnedVertices(const USkeletalMeshComponent* SkeletalMeshComponent, const int32 LODIndex,
+	TArray<FVector3f>& OutPositions);
+
+/** Gets Skin Weights Data from SkeletalMeshComponent */
+void GetSkinWeights(const USkeletalMesh* SkeletalMesh, const int32 LODIndex, 
+	TArray<VertexSkinWeightMax>& OutSkinWeights);
+
+/** Reduce Weights from MAX_TOTAL_INFLUENCES to 4 */
+void ReduceSkinWeights(const TArray<VertexSkinWeightMax>& InSkinWeights, TArray<VertexSkinWeightFour>& OutSkinWeights);
+
+/* Returns Number of RawBones (no virtual bones)*/
+int32 GetNumBones(const USkeletalMesh* SkeletalMesh);
+
+/* Returns RefPose Bone Transforms.
+   Only the RawBones are returned (no virtual bones)
+   The returned transforms are in ComponentSpace */
+void GetRefBoneTransforms(const USkeletalMesh* SkeletalMesh, TArray<FTransform>& OutTransforms);
+
+/* Returns Bone exist in the RawBone list (no virtual bones) */
+bool HasBone(const USkeletalMesh* SkeletalMesh, const FName& Bone);
+
+/* Returns Bone Names. 
+   Only the RawBones are returned (no virtual bones) */
+void GetBoneNames(const USkeletalMesh* SkeletalMesh, TArray<FName>& OutNames);
+
+} // end namespace AnimToTexture_Private
diff --git a/Engine/Plugins/Experimental/AnimToTexture/Source/AnimToTextureEditor/Public/AnimToTextureUtils.h b/Engine/Plugins/Experimental/AnimToTexture/Source/AnimToTextureEditor/Public/AnimToTextureUtils.h
new file mode 100644
--- /dev/null
+++ b/Engine/Plugins/Experimental/AnimToTexture/Source/AnimToTextureEditor/Public/AnimToTextureUtils.h
@@ -0,0 +1,278 @@
+// Copyright Epic Games, Inc. All Rights Reserved.
+
+#pragma once
+
+#include "AnimToTextureSkeletalMesh.h"
+#include "CoreMinimal.h"
+#include "TextureResource.h"
+#include "Engine/Texture.h"
+#include "Engine/Texture2D.h"
+
+
+namespace AnimToTexture_Private
+{
+
+struct FVector4u16
+{
+	uint16 X;
+	uint16 Y;
+	uint16 Z;
+	uint16 W;
+};
+
+struct FLowPrecision
+{
+	using ColorType = FColor;
+	static constexpr EPixelFormat PixelFormat = EPixelFormat::PF_B8G8R8A8;
+	static constexpr ETextureSourceFormat TextureSourceFormat = ETextureSourceFormat::TSF_BGRA8;
+	static constexpr TextureCompressionSettings CompressionSettings = TextureCompressionSettings::TC_VectorDisplacementmap;
+	static constexpr ColorType DefaultColor = FColor(0, 0, 0);		
+};
+
+struct FHighPrecision
+{
+	using ColorType = FVector4u16;
+	static constexpr EPixelFormat PixelFormat = EPixelFormat::PF_R16G16B16A16_UNORM;
+	static constexpr ETextureSourceFormat TextureSourceFormat = ETextureSourceFormat::TSF_RGBA16;
+	static constexpr TextureCompressionSettings CompressionSettings = TextureCompressionSettings::TC_HDR;
+	static constexpr ColorType DefaultColor = { 0, 0, 0, 0 };
+};
+
+/** Writes list of vectors into texture
+*   Note: They must be pre-normalized. */
+template<class V, class TextureSettings>
+bool WriteVectorsToTexture(const TArray<V>& Vectors,
+	const int32 NumFrames, const int32 RowsPerFrame,
+	const int32 Height, const int32 Width, 
+	UTexture2D* Texture);
+
+/* Writes list of skinweights into texture.
+*  The SkinWeights data is already in uint8 & uint16 format, no need for normalizing it.
+*/
+bool WriteSkinWeightsToTexture(const TArray<VertexSkinWeightFour>& SkinWeights,
+	const int32 RowsPerFrame,
+	const int32 Height, const int32 Width,
+	UTexture2D* Texture);
+
+/* Helper utility for writing 8 or 16 bits textures */
+template<class TextureSettings>
+bool WriteToTexture(UTexture2D* Texture, const uint32 Height, const uint32 Width, const TArray<typename TextureSettings::ColorType>& Data);
+
+template<class V /* FVector / FVector4 / FIntVector */, class C /* FColor / FLinearColor */>
+void VectorToColor(const V& Vector, C& Color);
+
+/* Decomposes Transform in Translation and AxisAndAngle */
+void DecomposeTransformation(const FTransform& Transform, FVector3f& OutTranslation, FVector4& OutRotation);
+void DecomposeTransformations(const TArray<FTransform>& Transforms, TArray<FVector3f>& OutTranslations, TArray<FVector4>& OutRotations);
+
+} // end namespace AnimToTexture_Private
+
+// ----------------------------------------------------------------------------
+
+template<>
+FORCEINLINE void AnimToTexture_Private::VectorToColor(const FVector3f& Vector, FColor& Color)
+{
+	const float ClampedX = FMath::Clamp(Vector.X, 0.f, 1.f);
+	const float ClampedY = FMath::Clamp(Vector.Y, 0.f, 1.f);
+	const float ClampedZ = FMath::Clamp(Vector.Z, 0.f, 1.f);
+
+	Color.R = (uint8)FMath::RoundToInt(ClampedX * 255.f);
+	Color.G = (uint8)FMath::RoundToInt(ClampedY * 255.f);
+	Color.B = (uint8)FMath::RoundToInt(ClampedZ * 255.f);
+	Color.A = 255;
+}
+
+template<>
+FORCEINLINE void AnimToTexture_Private::VectorToColor(const FVector3d& Vector, FColor& Color)
+{
+	VectorToColor((FVector3f)Vector, Color);
+}
+
+
+template<>
+FORCEINLINE void AnimToTexture_Private::VectorToColor(const FVector4& Vector, FColor& Color)
+{
+	const float ClampedX = FMath::Clamp(Vector.X, 0.f, 1.f);
+	const float ClampedY = FMath::Clamp(Vector.Y, 0.f, 1.f);
+	const float ClampedZ = FMath::Clamp(Vector.Z, 0.f, 1.f);
+	const float ClampedW = FMath::Clamp(Vector.W, 0.f, 1.f);
+	
+	Color.R = (uint8)FMath::RoundToInt(ClampedX * 255.f);
+	Color.G = (uint8)FMath::RoundToInt(ClampedY * 255.f);
+	Color.B = (uint8)FMath::RoundToInt(ClampedZ * 255.f);
+	Color.A = (uint8)FMath::RoundToInt(ClampedW * 255.f);
+}
+
+template<>
+FORCEINLINE void AnimToTexture_Private::VectorToColor(const FIntVector& Vector, FColor& Color)
+{
+	const int32 ClampedX = FMath::Clamp(Vector.X, 0, 1);
+	const int32 ClampedY = FMath::Clamp(Vector.Y, 0, 1);
+	const int32 ClampedZ = FMath::Clamp(Vector.Z, 0, 1);
+
+	Color.R = (uint8)(ClampedX * 255);
+	Color.G = (uint8)(ClampedY * 255);
+	Color.B = (uint8)(ClampedZ * 255);
+	Color.A = 255;
+}
+
+template<>
+FORCEINLINE void AnimToTexture_Private::VectorToColor(const FIntVector4& Vector, FColor& Color)
+{
+	const int32 ClampedX = FMath::Clamp(Vector.X, 0, 1);
+	const int32 ClampedY = FMath::Clamp(Vector.Y, 0, 1);
+	const int32 ClampedZ = FMath::Clamp(Vector.Z, 0, 1);
+	const int32 ClampedW = FMath::Clamp(Vector.W, 0, 1);
+
+	Color.R = (uint8)(ClampedX * 255);
+	Color.G = (uint8)(ClampedY * 255);
+	Color.B = (uint8)(ClampedZ * 255);
+	Color.A = (uint8)(ClampedW * 255);
+}
+
+template<>
+FORCEINLINE void AnimToTexture_Private::VectorToColor(const FVector3f& Vector, FLinearColor& Color)
+{
+	Color.R = Vector.X;
+	Color.G = Vector.Y;
+	Color.B = Vector.Z;
+	Color.A = 1.f;
+}
+
+template<>
+FORCEINLINE void AnimToTexture_Private::VectorToColor(const FVector4& Vector, FLinearColor& Color)
+{
+	Color.R = Vector.X;
+	Color.G = Vector.Y;
+	Color.B = Vector.Z;
+	Color.A = Vector.W;
+}
+
+template<>
+FORCEINLINE void AnimToTexture_Private::VectorToColor(const FVector3f& Vector, FVector4u16& Color)
+{
+	Color.X = FMath::RoundToInt(FMath::Clamp(Vector.X, 0.f, 1.f) * MAX_uint16);
+	Color.Y = FMath::RoundToInt(FMath::Clamp(Vector.Y, 0.f, 1.f) * MAX_uint16);
+	Color.Z = FMath::RoundToInt(FMath::Clamp(Vector.Z, 0.f, 1.f) * MAX_uint16);
+	Color.W = MAX_uint16;
+}
+
+template<>
+FORCEINLINE void AnimToTexture_Private::VectorToColor(const FVector4& Vector, FVector4u16& Color)
+{
+	Color.X = FMath::RoundToInt(FMath::Clamp(Vector.X, 0.f, 1.f) * MAX_uint16);
+	Color.Y = FMath::RoundToInt(FMath::Clamp(Vector.Y, 0.f, 1.f) * MAX_uint16);
+	Color.Z = FMath::RoundToInt(FMath::Clamp(Vector.Z, 0.f, 1.f) * MAX_uint16);
+	Color.W = FMath::RoundToInt(FMath::Clamp(Vector.W, 0.f, 1.f) * MAX_uint16);
+}
+
+
+template<class V, class TextureSettings>
+FORCEINLINE bool AnimToTexture_Private::WriteVectorsToTexture(const TArray<V>& Vectors,
+	const int32 NumFrames, const int32 RowsPerFrame,
+	const int32 Height, const int32 Width, UTexture2D* Texture)
+{
+	if (!Texture || !NumFrames)
+	{
+		return false;
+	}
+
+	// NumElements Per-Frame
+	const int32 NumElements = Vectors.Num() / NumFrames;
+
+	// Allocate PixelData.
+	TArray<typename TextureSettings::ColorType> Pixels;
+	Pixels.Init(TextureSettings::DefaultColor, Height * Width);
+
+	// Fillout Frame Data
+	for (int32 Frame = 0; Frame < NumFrames; Frame++)
+	{
+		const int32 BlockStart = RowsPerFrame * Width * Frame;
+
+		// Set Data.
+		for (int32 Index = 0; Index < NumElements; Index++)
+		{
+			const V& Vector = Vectors[NumElements * Frame + Index];
+			typename TextureSettings::ColorType& Pixel = Pixels[BlockStart + Index];
+
+			VectorToColor<V, typename TextureSettings::ColorType>(Vector, Pixel);
+		}
+	}
+
+	// Write to Texture
+	return WriteToTexture<TextureSettings>(Texture, Height, Width, Pixels);
+}
+
+template<class TextureSettings>
+FORCEINLINE bool AnimToTexture_Private::WriteToTexture(
+	UTexture2D* Texture,
+	const uint32 Height, const uint32 Width,
+	const TArray<typename TextureSettings::ColorType>& Pixels)
+{
+	if (!Texture)
+	{
+		return false;
+	}
+
+	// ------------------------------------------------------------------------
+	// Get Texture Platform
+	FTexturePlatformData* PlatformData = Texture->GetPlatformData();
+	if (!PlatformData)
+	{	
+		PlatformData = new FTexturePlatformData();
+		Texture->SetPlatformData(PlatformData);
+	}
+	PlatformData->SizeX = Width;
+	PlatformData->SizeY = Height;
+	PlatformData->SetNumSlices(1);
+	PlatformData->PixelFormat = TextureSettings::PixelFormat;
+	
+	// ------------------------------------------------------------------------
+	// Get First MipMap
+	//
+	FTexture2DMipMap* Mip;
+	if (PlatformData->Mips.IsEmpty())
+	{
+		Mip = new FTexture2DMipMap();
+		PlatformData->Mips.Add(Mip);
+	}
+	else
+	{
+		Mip = &PlatformData->Mips[0];
+
+	}
+	Mip->SizeX = Width;
+	Mip->SizeY = Height;
+	Mip->SizeZ = 1;
+	
+	// ------------------------------------------------------------------------
+	// Lock the Mipmap data so it can be modified
+	Mip->BulkData.Lock(LOCK_READ_WRITE);
+
+	// Reallocate MipMap
+	uint8* TextureData = (uint8*)Mip->BulkData.Realloc(Width * Height * sizeof(typename TextureSettings::ColorType));
+
+	// Copy the pixel data into the Texture data	
+	const uint8* PixelsData = (uint8*)Pixels.GetData();
+	FMemory::Memcpy(TextureData, PixelsData, Width * Height * sizeof(typename TextureSettings::ColorType));
+
+	// Unlock data
+	Mip->BulkData.Unlock();
+
+	// Initialize a new texture
+	Texture->Source.Init(Width, Height, 1, 1, TextureSettings::TextureSourceFormat, PixelsData);
+	
+	// Set parameters
+	Texture->SRGB = 0;
+	Texture->Filter = TextureFilter::TF_Nearest;
+	Texture->CompressionSettings = TextureSettings::CompressionSettings;
+	Texture->MipGenSettings = TextureMipGenSettings::TMGS_NoMipmaps;
+
+	// Update and Mark to Save.
+	Texture->UpdateResource();
+	Texture->MarkPackageDirty();
+
+	return true;
+}
+
+
