[AnimToTexture]
- RootTransform moved to DataAsset
- AnimationToTexture function returning bool
- GetRefBonePositionsAndRotations using SkeletalMesh instead of SkeletalMeshComponent
- Creating Tempory Component and Actor after running checks

diff --git a/Engine/Plugins/Experimental/AnimToTexture/Source/AnimToTexture/Public/AnimToTextureDataAsset.h b/Engine/Plugins/Experimental/AnimToTexture/Source/AnimToTexture/Public/AnimToTextureDataAsset.h
--- a/Engine/Plugins/Experimental/AnimToTexture/Source/AnimToTexture/Public/AnimToTextureDataAsset.h
+++ b/Engine/Plugins/Experimental/AnimToTexture/Source/AnimToTexture/Public/AnimToTextureDataAsset.h
@@ -214,6 +214,13 @@ public:
 
 	// ------------------------------------------------------
 	// Animation
+	
+	/**
+	* Adds transformation to baked textures. 
+	* This can be used for adding an offset to the animation.
+	*/
+	UPROPERTY(EditAnywhere, BlueprintReadWrite, Category = "Animation")
+	FTransform RootTransform;
 
 	/** 
 	* Bone used for Rigid Binding. The bone needs to be part of the RawBones. 
diff --git a/Engine/Plugins/Experimental/AnimToTexture/Source/AnimToTextureEditor/Private/AnimToTextureBPLibrary.cpp b/Engine/Plugins/Experimental/AnimToTexture/Source/AnimToTextureEditor/Private/AnimToTextureBPLibrary.cpp
--- a/Engine/Plugins/Experimental/AnimToTexture/Source/AnimToTextureEditor/Private/AnimToTextureBPLibrary.cpp
+++ b/Engine/Plugins/Experimental/AnimToTexture/Source/AnimToTextureEditor/Private/AnimToTextureBPLibrary.cpp
@@ -27,13 +27,11 @@ UAnimToTextureBPLibrary::UAnimToTextureBPLibrary(const FObjectInitializer& Objec
 
 }
 
-void UAnimToTextureBPLibrary::AnimationToTexture(UAnimToTextureDataAsset* DataAsset, const FTransform RootTransform, bool& bSuccess)
+bool UAnimToTextureBPLibrary::AnimationToTexture(UAnimToTextureDataAsset* DataAsset)
 {
-	bSuccess = false; 
-
 	if (!DataAsset)
 	{
-		return;
+		return false;
 	}
 
 	// Reset DataAsset Info Values
@@ -43,35 +41,35 @@ void UAnimToTextureBPLibrary::AnimationToTexture(UAnimToTextureDataAsset* DataAs
 	if (!DataAsset->GetStaticMesh())
 	{
 		UE_LOG(LogAnimToTextureEditor, Warning, TEXT("Invalid StaticMesh"));
-		return;
+		return false;
 	}
 
 	// Check SkeletalMesh
 	if (!DataAsset->GetSkeletalMesh())
 	{
 		UE_LOG(LogAnimToTextureEditor, Warning, TEXT("Invalid SkeletalMesh"));
-		return;
+		return false;
 	}
 
 	// Check Skeleton
 	if (!DataAsset->GetSkeletalMesh()->GetSkeleton())
 	{
 		UE_LOG(LogAnimToTextureEditor, Warning, TEXT("Invalid SkeletalMesh. No valid Skeleton found"));
-		return;
+		return false;
 	}
 
 	// Check StaticMesh LOD
 	if (!DataAsset->GetStaticMesh()->IsSourceModelValid(DataAsset->StaticLODIndex))
 	{
 		UE_LOG(LogAnimToTextureEditor, Warning, TEXT("Invalid StaticMesh LOD Index: %i"), DataAsset->StaticLODIndex);
-		return;
+		return false;
 	}
 
 	// Check SkeletalMesh LOD
 	if (!DataAsset->GetSkeletalMesh()->IsValidLODIndex(DataAsset->SkeletalLODIndex))
 	{
 		UE_LOG(LogAnimToTextureEditor, Warning, TEXT("Invalid SkeletalMesh LOD Index: %i"), DataAsset->SkeletalLODIndex);
-		return;
+		return false;
 	}
 
 	// Check Socket.
@@ -85,13 +83,13 @@ void UAnimToTextureBPLibrary::AnimationToTexture(UAnimToTextureDataAsset* DataAs
 		else
 		{
 			UE_LOG(LogAnimToTextureEditor, Warning, TEXT("Invalid Socket: %s"), *DataAsset->AttachToSocket.ToString());
-			return;
+			return false;
 		}
 	}
 	if (bValidSocket && DataAsset->Mode == EAnimToTextureMode::Vertex)
 	{
 		UE_LOG(LogAnimToTextureEditor, Warning, TEXT("Unable to use Socket in Vertex Mode. Use Bone Mode instead."));
-		return;
+		return false;
 	}
 
 	// Check if UVChannel is being used by the Lightmap UV
@@ -100,7 +98,7 @@ void UAnimToTextureBPLibrary::AnimationToTexture(UAnimToTextureDataAsset* DataAs
 		SourceModel.BuildSettings.DstLightmapIndex == DataAsset->UVChannel)
 	{
 		UE_LOG(LogAnimToTextureEditor, Warning, TEXT("Invalid UVChannel: %i. Already used by LightMap"), DataAsset->UVChannel);
-		return;
+		return false;
 	}
 
 	// Check Animations
@@ -111,11 +109,11 @@ void UAnimToTextureBPLibrary::AnimationToTexture(UAnimToTextureDataAsset* DataAs
 		if (AnimSequenceInfo.bEnabled && AnimSequence)
 		{
 			// Check Frame Range
-			if (AnimSequenceInfo.bUseCustomRange && 
-				(AnimSequenceInfo.EndFrame - AnimSequenceInfo.StartFrame) <= 0 )
-			{ 
+			if (AnimSequenceInfo.bUseCustomRange &&
+				(AnimSequenceInfo.EndFrame - AnimSequenceInfo.StartFrame) <= 0)
+			{
 				UE_LOG(LogAnimToTextureEditor, Warning, TEXT("Invalid Custom Range for AnimSequence: %s"), *AnimSequence->GetName());
-				return;
+				return false;
 			}
 
 			NumAnimations++;
@@ -124,7 +122,7 @@ void UAnimToTextureBPLibrary::AnimationToTexture(UAnimToTextureDataAsset* DataAs
 	if (!NumAnimations)
 	{
 		UE_LOG(LogAnimToTextureEditor, Warning, TEXT("No Animations found"));
-		return;
+		return false;
 	};
 
 	// ---------------------------------------------------------------------------		
@@ -149,26 +147,6 @@ void UAnimToTextureBPLibrary::AnimationToTexture(UAnimToTextureDataAsset* DataAs
 		DataAsset->GetStaticMesh(), DataAsset->StaticLODIndex,
 		DataAsset->GetSkeletalMesh(), DataAsset->SkeletalLODIndex, Mapping);
 
-	// --------------------------------------------------------------------------
-		
-	// Create Temp Actor
-	check(GEditor);
-	UWorld* World = GEditor->GetEditorWorldContext().World();
-	check(World);
-
-	AActor* Actor = World->SpawnActor<AActor>();
-	check(Actor);
-
-	// Create Temp SkeletalMesh Component
-	USkeletalMeshComponent* SkeletalMeshComponent = NewObject<USkeletalMeshComponent>(Actor);
-	check(SkeletalMeshComponent);
-	SkeletalMeshComponent->SetSkeletalMesh(DataAsset->GetSkeletalMesh());
-	SkeletalMeshComponent->SetForcedLOD(1); // Force to LOD0;
-	SkeletalMeshComponent->SetAnimationMode(EAnimationMode::AnimationSingleNode);
-	SkeletalMeshComponent->SetUpdateAnimationInEditor(true);
-	SkeletalMeshComponent->VisibilityBasedAnimTickOption = EVisibilityBasedAnimTickOption::AlwaysTickPoseAndRefreshBones;
-	SkeletalMeshComponent->RegisterComponent();
-
 	// ---------------------------------------------------------------------------
 	// Get Reference Skeleton Transforms
 	//
@@ -183,18 +161,18 @@ void UAnimToTextureBPLibrary::AnimationToTexture(UAnimToTextureDataAsset* DataAs
 	if (DataAsset->Mode == EAnimToTextureMode::Bone)
 	{
 		// Gets Ref Bone Position and Rotations.
-		NumBones = GetRefBonePositionsAndRotations(SkeletalMeshComponent,
+		NumBones = GetRefBonePositionsAndRotations(DataAsset->GetSkeletalMesh(),
 			BoneRefPositions, BoneRefRotations);
 
 		// NOTE: there is a limitation with the number of bones atm.
 		if (NumBones > 256)
 		{
 			UE_LOG(LogAnimToTextureEditor, Warning, TEXT("Invalid Number of Bones. There is a maximum of 256 bones"))
-			return;
+			return false;
 		}
 
 		// Get Bone Names (no virtual)
-		GetBoneNames(SkeletalMeshComponent->GetSkeletalMeshAsset(), BoneNames);
+		GetBoneNames(DataAsset->GetSkeletalMesh(), BoneNames);
 
 		// Make sure array sizes are correct.
 		check(BoneNames.Num() == NumBones);
@@ -203,7 +181,7 @@ void UAnimToTextureBPLibrary::AnimationToTexture(UAnimToTextureDataAsset* DataAs
 		if (bValidSocket && !BoneNames.Find(DataAsset->AttachToSocket, SocketIndex))
 		{
 			UE_LOG(LogAnimToTextureEditor, Warning, TEXT("Socket: %s not found in Raw Bone List"), *DataAsset->AttachToSocket.ToString());
-			return;
+			return false;
 		}
 
 		// Add RefPose 
@@ -212,6 +190,26 @@ void UAnimToTextureBPLibrary::AnimationToTexture(UAnimToTextureDataAsset* DataAs
 		BoneRotations.Append(BoneRefRotations);
 	}
 
+	// --------------------------------------------------------------------------
+
+	// Create Temp Actor
+	check(GEditor);
+	UWorld* World = GEditor->GetEditorWorldContext().World();
+	check(World);
+
+	AActor* Actor = World->SpawnActor<AActor>();
+	check(Actor);
+
+	// Create Temp SkeletalMesh Component
+	USkeletalMeshComponent* SkeletalMeshComponent = NewObject<USkeletalMeshComponent>(Actor);
+	check(SkeletalMeshComponent);
+	SkeletalMeshComponent->SetSkeletalMesh(DataAsset->GetSkeletalMesh());
+	SkeletalMeshComponent->SetForcedLOD(1); // Force to LOD0;
+	SkeletalMeshComponent->SetAnimationMode(EAnimationMode::AnimationSingleNode);
+	SkeletalMeshComponent->SetUpdateAnimationInEditor(true);
+	SkeletalMeshComponent->VisibilityBasedAnimTickOption = EVisibilityBasedAnimTickOption::AlwaysTickPoseAndRefreshBones;
+	SkeletalMeshComponent->RegisterComponent();
+
 	// ---------------------------------------------------------------------------
 	// Get Vertex Data (for all frames)
 	//		
@@ -296,7 +294,7 @@ void UAnimToTextureBPLibrary::AnimationToTexture(UAnimToTextureDataAsset* DataAs
 				
 				GetVertexDeltasAndNormals(SkeletalMeshComponent, DataAsset->SkeletalLODIndex,
 					Vertices, Normals, Mapping,
-					RootTransform,
+					DataAsset->RootTransform,
 					VertexFrameDeltas, VertexFrameNormals);
 					
 				VertexDeltas.Append(VertexFrameDeltas);
@@ -343,7 +341,7 @@ void UAnimToTextureBPLibrary::AnimationToTexture(UAnimToTextureDataAsset* DataAs
 	//
 	if (!DataAsset->NumFrames || !NumVertices)
 	{
-		return;
+		return false;
 	}
 
 	// ---------------------------------------------------------------------------
@@ -356,7 +354,7 @@ void UAnimToTextureBPLibrary::AnimationToTexture(UAnimToTextureDataAsset* DataAs
 								DataAsset->MaxHeight, DataAsset->MaxWidth, DataAsset->bEnforcePowerOfTwo))
 		{
 			UE_LOG(LogAnimToTextureEditor, Warning, TEXT("Vertex Animation data cannot be fit in a %ix%i texture."), DataAsset->MaxHeight, DataAsset->MaxWidth);
-			return;
+			return false;
 		}
 
 		// Normalize Vertex Data
@@ -402,7 +400,7 @@ void UAnimToTextureBPLibrary::AnimationToTexture(UAnimToTextureDataAsset* DataAs
 								DataAsset->MaxHeight, DataAsset->MaxWidth, DataAsset->bEnforcePowerOfTwo))
 		{
 			UE_LOG(LogAnimToTextureEditor, Warning, TEXT("Bone Animation data cannot be fit in a %ix%i texture."), DataAsset->MaxHeight, DataAsset->MaxWidth);
-			return;
+			return false;
 		}
 
 		// Write Bone Position and Rotation Textures
@@ -436,7 +434,7 @@ void UAnimToTextureBPLibrary::AnimationToTexture(UAnimToTextureDataAsset* DataAs
 								DataAsset->MaxHeight, DataAsset->MaxWidth, DataAsset->bEnforcePowerOfTwo))
 		{
 			UE_LOG(LogAnimToTextureEditor, Warning, TEXT("Weights Data cannot be fit in a %ix%i texture."), DataAsset->MaxHeight, DataAsset->MaxWidth);
-			return;
+			return false;
 		}
 
 		// Write Weights Texture
@@ -491,7 +489,7 @@ void UAnimToTextureBPLibrary::AnimationToTexture(UAnimToTextureDataAsset* DataAs
 	DataAsset->MarkPackageDirty();
 	
 	// All good here !
-	bSuccess = true;
+	return true;
 }
 
 
@@ -541,18 +539,23 @@ void UAnimToTextureBPLibrary::GetVertexDeltasAndNormals(const USkeletalMeshCompo
 }
 
 
-int32 UAnimToTextureBPLibrary::GetRefBonePositionsAndRotations(const USkeletalMeshComponent* SkeletalMeshComponent,
+int32 UAnimToTextureBPLibrary::GetRefBonePositionsAndRotations(const USkeletalMesh* SkeletalMesh,
 	TArray<FVector3f>& OutBoneRefPositions, TArray<FVector4>& OutBoneRefRotations)
 {
 	OutBoneRefPositions.Reset();
 	OutBoneRefRotations.Reset();
 
+	if (!SkeletalMesh)
+	{
+		return 0;
+	}
+
 	// Get Number of RawBones (no virtual)
-	const int32 NumBones = GetNumBones(SkeletalMeshComponent->GetSkeletalMeshAsset());
+	const int32 NumBones = GetNumBones(SkeletalMesh);
 	
 	// Get Raw Ref Bone (no virtual)
 	TArray<FTransform> RefBoneTransforms;
-	GetRefBoneTransforms(SkeletalMeshComponent->GetSkeletalMeshAsset(), RefBoneTransforms);
+	GetRefBoneTransforms(SkeletalMesh, RefBoneTransforms);
 	DecomposeTransformations(RefBoneTransforms, OutBoneRefPositions, OutBoneRefRotations);
 
 	return NumBones;
diff --git a/Engine/Plugins/Experimental/AnimToTexture/Source/AnimToTextureEditor/Public/AnimToTextureBPLibrary.h b/Engine/Plugins/Experimental/AnimToTexture/Source/AnimToTextureEditor/Public/AnimToTextureBPLibrary.h
--- a/Engine/Plugins/Experimental/AnimToTexture/Source/AnimToTextureEditor/Public/AnimToTextureBPLibrary.h
+++ b/Engine/Plugins/Experimental/AnimToTexture/Source/AnimToTextureEditor/Public/AnimToTextureBPLibrary.h
@@ -12,11 +12,10 @@
 class UStaticMesh;
 class USkeletalMesh;
 
-/*
-* TODO: 
-*	- Right now it is saving data per-vertex instead of per-point.
-*     This will require more pixels if the mesh has lots of material slots or uv-shells.
-*     FStaticMeshOperations::FindOverlappingCorners(FOverlappingCorners& OutOverlappingCorners, const FMeshDescription& MeshDescription, float ComparisonThreshold)
+/**
+* Converts a series of animation frames into a Vertex Animation Texture (VAT)
+* 
+* The texture can store Vertex positions and normals or Bone positions and rotations.
 */
 UCLASS()
 class UAnimToTextureBPLibrary : public UBlueprintFunctionLibrary
@@ -29,7 +28,7 @@ class UAnimToTextureBPLibrary : public UBlueprintFunctionLibrary
 	* Bakes Animation Data into Textures.
 	*/
 	UFUNCTION(BlueprintCallable, meta = (Category = "AnimToTexture"))
-	static void AnimationToTexture(UAnimToTextureDataAsset* DataAsset, const FTransform RootTransform, bool& bSuccess);
+	static UPARAM(DisplayName="bSuccess") bool AnimationToTexture(UAnimToTextureDataAsset* DataAsset);
 
 	/** 
 	* Utility for converting SkeletalMesh into a StaticMesh
@@ -41,7 +40,7 @@ class UAnimToTextureBPLibrary : public UBlueprintFunctionLibrary
 	* Utility for setting a StaticMesh LightMapIndex.
 	*/
 	UFUNCTION(BlueprintCallable, Category = "AnimToTexture")
-	static bool SetLightMapIndex(UStaticMesh* StaticMesh, const int32 LODIndex, const int32 LightmapIndex=1, bool bGenerateLightmapUVs=true);
+	static UPARAM(DisplayName = "bSuccess") bool SetLightMapIndex(UStaticMesh* StaticMesh, const int32 LODIndex, const int32 LightmapIndex=1, bool bGenerateLightmapUVs=true);
 
 	/**
 	* Updates a material's parameters to match those of an AnimToTexture DataAsset
@@ -66,7 +65,7 @@ private:
 		TArray<FVector3f>& OutVertexDeltas, TArray<FVector3f>& OutVertexNormals);
 
 	// Gets RefPose Bone Position and Rotations.
-	static int32 GetRefBonePositionsAndRotations(const USkeletalMeshComponent* SkeletalMeshComponent, 
+	static int32 GetRefBonePositionsAndRotations(const USkeletalMesh* SkeletalMesh, 
 		TArray<FVector3f>& OutBoneRefPositions, TArray<FVector4>& OutBoneRefRotations);
 
 	// Gets Bone Position and Rotations for Current Pose.	
