[AnimToTexture]
- Right click on AnimToTextureDataAsset
- Registered MessageLog

diff --git a/Engine/Plugins/Experimental/AnimToTexture/Source/AnimToTextureEditor/AnimToTextureEditor.Build.cs b/Engine/Plugins/Experimental/AnimToTexture/Source/AnimToTextureEditor/AnimToTextureEditor.Build.cs
--- a/Engine/Plugins/Experimental/AnimToTexture/Source/AnimToTextureEditor/AnimToTextureEditor.Build.cs
+++ b/Engine/Plugins/Experimental/AnimToTexture/Source/AnimToTextureEditor/AnimToTextureEditor.Build.cs
@@ -26,7 +26,8 @@ public class AnimToTextureEditor : ModuleRules
 			new string[]
 			{
 				"Core",
-				"MaterialEditor"
+				"MaterialEditor",
+				"MessageLog"
 				// ... add other public dependencies that you statically link with here ...
 			}
 			);
diff --git a/Engine/Plugins/Experimental/AnimToTexture/Source/AnimToTextureEditor/Private/AnimToTextureAssetActions.cpp b/Engine/Plugins/Experimental/AnimToTexture/Source/AnimToTextureEditor/Private/AnimToTextureAssetActions.cpp
new file mode 100644
--- /dev/null
+++ b/Engine/Plugins/Experimental/AnimToTexture/Source/AnimToTextureEditor/Private/AnimToTextureAssetActions.cpp
@@ -0,0 +1,82 @@
+// Copyright Epic Games, Inc. All Rights Reserved.
+
+#include "AnimToTextureAssetActions.h"
+#include "AnimToTextureDataAsset.h"
+#include "AnimToTextureBPLibrary.h"
+
+#include "Styling/SlateIconFinder.h"
+#include "Widgets/Images/SImage.h"
+#include "Widgets/Layout/SBorder.h"
+
+FText FAnimToTextureAssetActions::GetName() const
+{
+	return NSLOCTEXT("AssetTypeActions", "AnimToTextureAssetActions", "AnimToTexture");
+}
+
+FColor FAnimToTextureAssetActions::GetTypeColor() const
+{
+	return FColor::Blue;
+}
+
+TSharedPtr<SWidget> FAnimToTextureAssetActions::GetThumbnailOverlay(const FAssetData& AssetData) const
+{
+	const FSlateBrush* Icon = FSlateIconFinder::FindIconBrushForClass(UAnimToTextureDataAsset::StaticClass());;
+
+	return SNew(SBorder)
+		.BorderImage(FAppStyle::GetNoBrush())
+		.Visibility(EVisibility::HitTestInvisible)
+		.Padding(FMargin(0.0f, 0.0f, 0.0f, 3.0f))
+		.HAlign(HAlign_Right)
+		.VAlign(VAlign_Bottom)
+		[
+			SNew(SImage)
+			.Image(Icon)
+		];
+}
+
+UClass* FAnimToTextureAssetActions::GetSupportedClass() const
+{
+	return UAnimToTextureDataAsset::StaticClass();
+}
+
+uint32 FAnimToTextureAssetActions::GetCategories()
+{
+	return EAssetTypeCategories::Animation;
+}
+
+const TArray<FText>& FAnimToTextureAssetActions::GetSubMenus() const
+{
+	static const TArray<FText> SubMenus
+	{
+		NSLOCTEXT("AnimToTextureAssetActions", "AnimDeformersSubMenu", "Deformers")
+	};
+	return SubMenus;
+}
+
+void FAnimToTextureAssetActions::GetActions( const TArray<UObject*>& InObjects, FMenuBuilder& MenuBuilder )
+{
+	TArray<TWeakObjectPtr<UAnimToTextureDataAsset>> Objects = GetTypedWeakObjectPtrs<UAnimToTextureDataAsset>(InObjects);
+
+	MenuBuilder.AddMenuEntry(
+		NSLOCTEXT("AnimToTextureAssetActions","AnimToTexture_Run", "Run Animation To Texture"),
+		NSLOCTEXT("AnimToTextureAssetActions", "AnimToTexture_RunTooltip", "Creates Vertex Animation Textures (VAT)"),
+		FSlateIcon(),
+		FUIAction(
+			FExecuteAction::CreateSP( this, &FAnimToTextureAssetActions::RunAnimToTexture, Objects ),
+			FCanExecuteAction()
+			)
+		);
+
+	MenuBuilder.AddMenuSeparator();
+}
+
+void FAnimToTextureAssetActions::RunAnimToTexture(TArray<TWeakObjectPtr<UAnimToTextureDataAsset>> Objects)
+{
+	for (auto ObjIt = Objects.CreateConstIterator(); ObjIt; ++ObjIt)
+	{
+		if (UAnimToTextureDataAsset* DataAsset = (*ObjIt).Get())
+		{
+			UAnimToTextureBPLibrary::AnimationToTexture(DataAsset);
+		}
+	}
+}
diff --git a/Engine/Plugins/Experimental/AnimToTexture/Source/AnimToTextureEditor/Private/AnimToTextureEditorModule.cpp b/Engine/Plugins/Experimental/AnimToTexture/Source/AnimToTextureEditor/Private/AnimToTextureEditorModule.cpp
--- a/Engine/Plugins/Experimental/AnimToTexture/Source/AnimToTextureEditor/Private/AnimToTextureEditorModule.cpp
+++ b/Engine/Plugins/Experimental/AnimToTexture/Source/AnimToTextureEditor/Private/AnimToTextureEditorModule.cpp
@@ -1,6 +1,10 @@
 // Copyright Epic Games, Inc. All Rights Reserved.
 
 #include "AnimToTextureEditorModule.h"
+#include "AnimToTextureAssetActions.h"
+
+#include "AssetToolsModule.h"
+#include "MessageLogModule.h"
 
 #define LOCTEXT_NAMESPACE "FAnimToTextureEditorModule"
 
@@ -8,17 +12,27 @@ DEFINE_LOG_CATEGORY(LogAnimToTextureEditor);
 
 void FAnimToTextureEditorModule::StartupModule()
 {
-	// This code will execute after your module is loaded into memory; the exact timing is specified in the .uplugin file per-module
-	
+	// Asset actions
+	FAssetToolsModule& AssetToolsModule = FModuleManager::LoadModuleChecked< FAssetToolsModule >("AssetTools");	
+	TSharedPtr<FAnimToTextureAssetActions> AnimToTextureTypeActions = MakeShareable( new FAnimToTextureAssetActions);
+	AssetToolsModule.Get().RegisterAssetTypeActions(AnimToTextureTypeActions.ToSharedRef() );
+
+	// Register Log
+	FMessageLogModule& MessageLogModule = FModuleManager::LoadModuleChecked<FMessageLogModule>("MessageLog");
+	FMessageLogInitializationOptions InitOptions;
+	InitOptions.bShowFilters = true;
+	InitOptions.bShowPages = false;
+	InitOptions.bAllowClear = true;
+	MessageLogModule.RegisterLogListing("AnimToTextureLog", LOCTEXT("AnimToTextureLog", "AnimToTexture Log"), InitOptions);
 }
 
 void FAnimToTextureEditorModule::ShutdownModule()
 {
-	// This function may be called during shutdown to clean up your module.  For modules that support dynamic reloading,
-	// we call this function before unloading the module.
-	
+	FMessageLogModule& MessageLogModule = FModuleManager::LoadModuleChecked<FMessageLogModule>("MessageLog");
+	MessageLogModule.UnregisterLogListing("AnimToTextureLog");
 }
 
+IMPLEMENT_MODULE(FAnimToTextureEditorModule, AnimToTextureEditor)
+
 #undef LOCTEXT_NAMESPACE
-	
-IMPLEMENT_MODULE(FAnimToTextureEditorModule, AnimToTextureEditor)
\ No newline at end of file
+	
\ No newline at end of file
diff --git a/Engine/Plugins/Experimental/AnimToTexture/Source/AnimToTextureEditor/Public/AnimToTextureAssetActions.h b/Engine/Plugins/Experimental/AnimToTexture/Source/AnimToTextureEditor/Public/AnimToTextureAssetActions.h
new file mode 100644
--- /dev/null
+++ b/Engine/Plugins/Experimental/AnimToTexture/Source/AnimToTextureEditor/Public/AnimToTextureAssetActions.h
@@ -0,0 +1,30 @@
+// Copyright Epic Games, Inc. All Rights Reserved.
+
+#pragma once
+
+#include "AssetTypeActions_Base.h"
+#include "AnimToTextureDataAsset.h"
+
+class FMenuBuilder;
+class IToolkitHost;
+class UClass;
+class UObject;
+
+
+class FAnimToTextureAssetActions : public FAssetTypeActions_Base
+{
+public:
+	
+	// IAssetTypeActions Implementation
+	FText GetName() const override;
+	FColor GetTypeColor() const override;
+	UClass* GetSupportedClass() const override;	
+	void GetActions( const TArray<UObject*>& InObjects, FMenuBuilder& MenuBuilder ) override;
+	uint32 GetCategories() override;
+	const TArray<FText>& GetSubMenus() const override;
+	TSharedPtr<SWidget> GetThumbnailOverlay(const FAssetData& AssetData) const override;
+
+private:
+	void RunAnimToTexture(TArray<TWeakObjectPtr<UAnimToTextureDataAsset>> Objects);
+	
+};
