[AnimToTexture] Updating MaterialInstance Params with AssetActions (right-click)

diff --git a/Engine/Plugins/Experimental/AnimToTexture/Source/AnimToTextureEditor/Private/AnimToTextureAssetActions.cpp b/Engine/Plugins/Experimental/AnimToTexture/Source/AnimToTextureEditor/Private/AnimToTextureAssetActions.cpp
--- a/Engine/Plugins/Experimental/AnimToTexture/Source/AnimToTextureEditor/Private/AnimToTextureAssetActions.cpp
+++ b/Engine/Plugins/Experimental/AnimToTexture/Source/AnimToTextureEditor/Private/AnimToTextureAssetActions.cpp
@@ -4,6 +4,7 @@
 #include "AnimToTextureDataAsset.h"
 #include "AnimToTextureBPLibrary.h"
 
+#include "Materials/MaterialInstanceConstant.h"
 #include "Styling/SlateIconFinder.h"
 #include "Widgets/Images/SImage.h"
 #include "Widgets/Layout/SBorder.h"
@@ -76,7 +77,24 @@ void FAnimToTextureAssetActions::RunAnimToTexture(TArray<TWeakObjectPtr<UAnimToT
 	{
 		if (UAnimToTextureDataAsset* DataAsset = (*ObjIt).Get())
 		{
+			// Create UVs and Textures
 			UAnimToTextureBPLibrary::AnimationToTexture(DataAsset);
+
+			// Update Material Instances (if Possible)
+			if (UStaticMesh* StaticMesh = DataAsset->GetStaticMesh())
+			{
+				for (FStaticMaterial& StaticMaterial : StaticMesh->GetStaticMaterials())
+				{
+					if (UMaterialInstanceConstant* MaterialInstanceConstant = Cast<UMaterialInstanceConstant>(StaticMaterial.MaterialInterface))
+					{
+						UAnimToTextureBPLibrary::UpdateMaterialInstanceFromDataAsset(DataAsset, MaterialInstanceConstant,
+							/*bAutoPlay*/ true, 
+							/*AnimationIndex*/ 0, 
+							EAnimToTextureNumBoneInfluences::Four,
+							EMaterialParameterAssociation::LayerParameter);
+					}
+				}
+			}
 		}
 	}
 }
