[AnimToTexture]
- Smooth deformation
- Fixed some issues custom StartFrame and EndFrame

diff --git a/Engine/Plugins/Experimental/AnimToTexture/AnimToTexture.uplugin b/Engine/Plugins/Experimental/AnimToTexture/AnimToTexture.uplugin
--- a/Engine/Plugins/Experimental/AnimToTexture/AnimToTexture.uplugin
+++ b/Engine/Plugins/Experimental/AnimToTexture/AnimToTexture.uplugin
@@ -1,10 +1,10 @@
 {
 	"FileVersion": 3,
 	"Version": 2,
-	"VersionName": "2.0",
+	"VersionName": "2.1",
 	"FriendlyName": "AnimToTexture",
 	"Description": "Converts SkeletalMesh Animations into Textures",
-	"Category": "Other",
+	"Category": "Animation",
 	"CreatedBy": "Epic Games, Inc.",
 	"CreatedByURL": "https://epicgames.com",
 	"DocsURL": "",
diff --git a/Engine/Plugins/Experimental/AnimToTexture/Source/AnimToTexture/Public/AnimToTextureDataAsset.h b/Engine/Plugins/Experimental/AnimToTexture/Source/AnimToTexture/Public/AnimToTextureDataAsset.h
--- a/Engine/Plugins/Experimental/AnimToTexture/Source/AnimToTexture/Public/AnimToTextureDataAsset.h
+++ b/Engine/Plugins/Experimental/AnimToTexture/Source/AnimToTexture/Public/AnimToTextureDataAsset.h
@@ -77,15 +77,18 @@ struct FAnimToTextureAnimSequenceInfo
 
 	UPROPERTY(EditAnywhere, Category = Default, BlueprintReadWrite, meta = (EditCondition = "bEnabled", EditConditionHides))
 	TObjectPtr<UAnimSequence> AnimSequence = nullptr;
-
+	
+	/* Use Custom FrameRange */
 	UPROPERTY(EditAnywhere, Category = Default, BlueprintReadWrite, meta = (EditCondition = "bEnabled", EditConditionHides))
 	bool bUseCustomRange = false;
 
+	/* Animation Start Frame */
 	UPROPERTY(EditAnywhere, Category = Default, BlueprintReadWrite, meta = (EditCondition = "bEnabled && bUseCustomRange", EditConditionHides))
 	int32 StartFrame = 0;
 
+	/* Animation End Frame (Inclusive) */
 	UPROPERTY(EditAnywhere, Category = Default, BlueprintReadWrite, meta = (EditCondition = "bEnabled && bUseCustomRange", EditConditionHides))
-	int32 EndFrame = 1;
+	int32 EndFrame = 0;
 
 };
 
diff --git a/Engine/Plugins/Experimental/AnimToTexture/Source/AnimToTextureEditor/Private/AnimToTextureBPLibrary.cpp b/Engine/Plugins/Experimental/AnimToTexture/Source/AnimToTextureEditor/Private/AnimToTextureBPLibrary.cpp
--- a/Engine/Plugins/Experimental/AnimToTexture/Source/AnimToTextureEditor/Private/AnimToTextureBPLibrary.cpp
+++ b/Engine/Plugins/Experimental/AnimToTexture/Source/AnimToTextureEditor/Private/AnimToTextureBPLibrary.cpp
@@ -110,7 +110,9 @@ bool UAnimToTextureBPLibrary::AnimationToTexture(UAnimToTextureDataAsset* DataAs
 		{
 			// Check Frame Range
 			if (AnimSequenceInfo.bUseCustomRange &&
-				(AnimSequenceInfo.EndFrame - AnimSequenceInfo.StartFrame) <= 0)
+				(AnimSequenceInfo.StartFrame < 0 ||
+				 AnimSequenceInfo.EndFrame > AnimSequence->GetNumberOfSampledKeys() - 1 ||
+				 AnimSequenceInfo.EndFrame - AnimSequenceInfo.StartFrame < 0 ))
 			{
 				UE_LOG(LogAnimToTextureEditor, Warning, TEXT("Invalid Custom Range for AnimSequence: %s"), *AnimSequence->GetName());
 				return false;
@@ -131,21 +133,15 @@ bool UAnimToTextureBPLibrary::AnimationToTexture(UAnimToTextureDataAsset* DataAs
 	//       Since they dont have same number of points.
 	//
 
-	// Get Static Vertices 
-	// NOTE: these are the RenderData Vertices.
-	TArray<FVector3f> Vertices;
-	TArray<FVector3f> Normals;
-	const int32 NumVertices = GetVertices(DataAsset->GetStaticMesh(), DataAsset->StaticLODIndex, Vertices, Normals);
-
-	// Get SkeletalMesh Vertices 
-	TArray<FVector3f> SkelVertices;
-	GetVertices(DataAsset->GetSkeletalMesh(), DataAsset->SkeletalLODIndex, SkelVertices);
+	// Get SourceMeshToDriverMesh
+	FSourceMeshToDriverMesh Mapping(DataAsset->GetStaticMesh(), DataAsset->StaticLODIndex, DataAsset->GetSkeletalMesh(), DataAsset->SkeletalLODIndex);
 	
-	// Get Mapping
-	TArray<FVertexToMeshMapping> Mapping;
-	FVertexToMeshMapping::Create(
-		DataAsset->GetStaticMesh(), DataAsset->StaticLODIndex,
-		DataAsset->GetSkeletalMesh(), DataAsset->SkeletalLODIndex, Mapping);
+	// Get Number of Source Vertices (StaticMesh)
+	const int32 NumVertices = Mapping.GetNumSourceVertices();
+	if (!NumVertices)
+	{
+		return false;
+	}
 
 	// ---------------------------------------------------------------------------
 	// Get Reference Skeleton Transforms
@@ -184,6 +180,8 @@ bool UAnimToTextureBPLibrary::AnimationToTexture(UAnimToTextureDataAsset* DataAs
 			return false;
 		}
 
+		// TODO: SocketIndex can only be < TNumericLimits<uint16>::Max()
+		
 		// Add RefPose 
 		// Note: this is added in the first frame of the Bone Position and Rotation Textures
 		BonePositions.Append(BoneRefPositions);
@@ -218,8 +216,6 @@ bool UAnimToTextureBPLibrary::AnimationToTexture(UAnimToTextureDataAsset* DataAs
 	
 	// Get Animation Frames Data
 	//
-	const float SampleInterval = 1.f / DataAsset->SampleRate;
-
 	for (const FAnimToTextureAnimSequenceInfo& AnimSequenceInfo : DataAsset->AnimSequences)
 	{
 		UAnimSequence* AnimSequence = AnimSequenceInfo.AnimSequence;
@@ -230,6 +226,7 @@ bool UAnimToTextureBPLibrary::AnimationToTexture(UAnimToTextureDataAsset* DataAs
 		}
 		
 		// Make sure SkeletalMesh is compatible with AnimSequence
+		// TODO: move this to early Checks.
 		if (!SkeletalMeshComponent->GetSkeletalMeshAsset()->GetSkeleton()->IsCompatibleForEditor(AnimSequence->GetSkeleton()))
 		{
 			UE_LOG(LogAnimToTextureEditor, Warning, TEXT("Invalid AnimSequence: %s for given SkeletalMesh: %s"), *AnimSequence->GetFName().ToString(), *SkeletalMeshComponent->GetSkeletalMeshAsset()->GetFName().ToString());
@@ -237,43 +234,40 @@ bool UAnimToTextureBPLibrary::AnimationToTexture(UAnimToTextureDataAsset* DataAs
 		}
 		// Set AnimSequence
 		else
-		{			
+		{
 			SkeletalMeshComponent->SetAnimation(AnimSequence);
 		}
 
 		// -----------------------------------------------------------------------------------
 		// Get Number of Frames
 		//
-		int32 StartFrame;
-		int32 EndFrame;
+		int32 AnimStartFrame;
+		int32 AnimEndFrame;
 		
 		// Get Range from AnimSequence
 		if (!AnimSequenceInfo.bUseCustomRange)
 		{
-			StartFrame = 0;
-			EndFrame = AnimSequence->GetNumberOfSampledKeys() - 1; // AnimSequence->GetNumberOfFrames();
+			AnimStartFrame = 0;
+			AnimEndFrame = AnimSequence->GetNumberOfSampledKeys() - 1; // AnimSequence->GetNumberOfFrames();
 		}
 		// Get Range from DataAsset
 		else
 		{
-			StartFrame = AnimSequenceInfo.StartFrame;
-			EndFrame = AnimSequenceInfo.EndFrame;
-		}
-		if (EndFrame - StartFrame <= 0)
-		{
-			continue;
+			AnimStartFrame = AnimSequenceInfo.StartFrame;
+			AnimEndFrame = AnimSequenceInfo.EndFrame;
 		}
 		
 		// ---------------------------------------------------------------------------
 		// 
-		float Time = 0.f;
-		const float EndTime = AnimSequence->GetTimeAtFrame(EndFrame);
+		const int32 AnimNumFrames = AnimEndFrame - AnimStartFrame + 1;
+		const float AnimStartTime = AnimSequence->GetTimeAtFrame(AnimStartFrame);
 
 		int32 SampleIndex = 0;
+		const float SampleInterval = 1.f / DataAsset->SampleRate;
 
-		while (Time < EndTime)
+		while (SampleIndex < AnimNumFrames)
 		{
-			Time = FMath::Clamp(SampleIndex * SampleInterval, 0.f, EndTime);
+			const float Time = AnimStartTime + ((float)SampleIndex * SampleInterval);
 			SampleIndex++;
 
 			// Go To Time
@@ -293,8 +287,7 @@ bool UAnimToTextureBPLibrary::AnimationToTexture(UAnimToTextureDataAsset* DataAs
 				TArray<FVector3f> VertexFrameNormals;
 				
 				GetVertexDeltasAndNormals(SkeletalMeshComponent, DataAsset->SkeletalLODIndex,
-					Vertices, Normals, Mapping,
-					DataAsset->RootTransform,
+					Mapping, DataAsset->RootTransform,
 					VertexFrameDeltas, VertexFrameNormals);
 					
 				VertexDeltas.Append(VertexFrameDeltas);
@@ -323,23 +316,22 @@ bool UAnimToTextureBPLibrary::AnimationToTexture(UAnimToTextureDataAsset* DataAs
 		// Store Anim Info Data
 		FAnimToTextureAnimInfo AnimInfo;
 		AnimInfo.StartFrame = DataAsset->NumFrames;
-		AnimInfo.EndFrame = DataAsset->NumFrames + EndFrame;
+		AnimInfo.EndFrame = DataAsset->NumFrames + AnimNumFrames - 1;
 		DataAsset->Animations.Add(AnimInfo);
 
 		// Accumulate Frames
-		DataAsset->NumFrames += SampleIndex;
+		DataAsset->NumFrames += AnimNumFrames;
 	} // End Anim
 		
-	// Destroy Temp Component
+	// Destroy Temp Component & Actor
 	SkeletalMeshComponent->UnregisterComponent();
 	SkeletalMeshComponent->DestroyComponent();
-
 	Actor->Destroy();
 	
 	// ---------------------------------------------------------------------------
 	// Nothing to do here ...
 	//
-	if (!DataAsset->NumFrames || !NumVertices)
+	if (!DataAsset->NumFrames)
 	{
 		return false;
 	}
@@ -428,7 +420,7 @@ bool UAnimToTextureBPLibrary::AnimationToTexture(UAnimToTextureDataAsset* DataAs
 
 		// ---------------------------------------------------------------------------
 
-		// Find Best Resolution for Weights Data
+		// Find Best Resolution for Bone Weights Texture
 		if (!FindBestResolution(2, NumVertices, 
 								Height, Width, DataAsset->BoneWeightRowsPerFrame, 
 								DataAsset->MaxHeight, DataAsset->MaxWidth, DataAsset->bEnforcePowerOfTwo))
@@ -441,18 +433,14 @@ bool UAnimToTextureBPLibrary::AnimationToTexture(UAnimToTextureDataAsset* DataAs
 		{
 			TArray<TVertexSkinWeight<4>> SkinWeights;
 
-			// Store Influence Weights
+			// Reduce BoneWeights to 4 Influences.
 			if (!bValidSocket)
-			{
-				// Get SkeletalMesh Weights
-				TArray<VertexSkinWeightMax> SkeletalMeshSkinWeights;
-				GetSkinWeights(DataAsset->GetSkeletalMesh(), DataAsset->SkeletalLODIndex, SkeletalMeshSkinWeights);
-
-				// Interpolate Weights with Mapping
+			{	
+				// Project SkinWeights from SkeletalMesh to StaticMesh
 				TArray<VertexSkinWeightMax> StaticMeshSkinWeights;
-				FVertexToMeshMapping::InterpolateSkinWeights(Mapping, SkeletalMeshSkinWeights, StaticMeshSkinWeights);
-				
-				// Reduce SkeletalMesh Weights to 4 highest influences.
+				Mapping.ProjectSkinWeights(StaticMeshSkinWeights);
+
+				// Reduce Weights to 4 highest influences.
 				ReduceSkinWeights(StaticMeshSkinWeights, SkinWeights);
 			}
 
@@ -495,43 +483,41 @@ bool UAnimToTextureBPLibrary::AnimationToTexture(UAnimToTextureDataAsset* DataAs
 
 // 
 void UAnimToTextureBPLibrary::GetVertexDeltasAndNormals(const USkeletalMeshComponent* SkeletalMeshComponent, const int32 LODIndex, 
-	const TArray<FVector3f>& Vertices, const TArray<FVector3f>& Normals,  const TArray<FVertexToMeshMapping>& Mapping,
+	const AnimToTexture_Private::FSourceMeshToDriverMesh& SourceMeshToDriverMesh,
 	const FTransform RootTransform,
 	TArray<FVector3f>& OutVertexDeltas, TArray<FVector3f>& OutVertexNormals)
 {
 	OutVertexDeltas.Reset();
 	OutVertexNormals.Reset();
-
-	check(Vertices.Num() == Mapping.Num());
-	const int32 NumVertices = Vertices.Num();
-
+		
 	// Get Deformed vertices at current frame
 	TArray<FVector3f> SkinnedVertices;
 	GetSkinnedVertices(SkeletalMeshComponent, LODIndex, SkinnedVertices);
+	
+	// Get Source Vertices (StaticMesh)
+	TArray<FVector3f> SourceVertices;
+	const int32 NumVertices = SourceMeshToDriverMesh.GetSourceVertices(SourceVertices);
+
+	// Deform Source Vertices with DriverMesh (SkeletalMesh
+	TArray<FVector3f> DeformedVertices;
+	TArray<FVector3f> DeformedNormals;
+	SourceMeshToDriverMesh.DeformVerticesAndNormals(SkinnedVertices, DeformedVertices, DeformedNormals);
 
 	// Allocate
+	check(DeformedVertices.Num() == NumVertices && DeformedNormals.Num() == NumVertices);
 	OutVertexDeltas.SetNumUninitialized(NumVertices);
 	OutVertexNormals.SetNumUninitialized(NumVertices);
 
-	// Loop thru static vertices and find the mapped SkeletalMesh Vertex
+	// Transform Vertices and Normals with RootTransform
 	for (int32 VertexIndex = 0; VertexIndex < NumVertices; VertexIndex++)
 	{
-		// Get Static/SkeletalMesh Mapping
-		const FVertexToMeshMapping& VertexMapping = Mapping[VertexIndex];
-
-		// Transform Vertex
-		// (assuming indices are valid)
-		const FIntVector3& Triangle = VertexMapping.Triangle;
-		const FVector3f& A = SkinnedVertices[Triangle.X];
-		const FVector3f& B = SkinnedVertices[Triangle.Y];
-		const FVector3f& C = SkinnedVertices[Triangle.Z];
-
-		const FVector3f VertexPosition = VertexMapping.TransformPosition(Vertices[VertexIndex], A, B, C);
-		const FVector3f VertexNormal = VertexMapping.TransformVector(Normals[VertexIndex], A, B, C);
-
+		const FVector3f& SourceVertex   = SourceVertices[VertexIndex];
+		const FVector3f& DeformedVertex = DeformedVertices[VertexIndex];
+		const FVector3f& DeformedNormal = DeformedNormals[VertexIndex];
+	
 		// Transform Position and Delta with RootTransform
-		const FVector3f TransformedVertexDelta = ((FVector3f)RootTransform.TransformPosition((FVector)VertexPosition)) - Vertices[VertexIndex];
-		const FVector3f TransformedVertexNormal = (FVector3f)RootTransform.TransformVector((FVector)VertexNormal);
+		const FVector3f TransformedVertexDelta = ((FVector3f)RootTransform.TransformPosition((FVector)DeformedVertex)) - SourceVertex;
+		const FVector3f TransformedVertexNormal = (FVector3f)RootTransform.TransformVector((FVector)DeformedNormal);
 		
 		OutVertexDeltas[VertexIndex] = TransformedVertexDelta;
 		OutVertexNormals[VertexIndex] = TransformedVertexNormal;
diff --git a/Engine/Plugins/Experimental/AnimToTexture/Source/AnimToTextureEditor/Private/AnimToTextureMeshMapping.cpp b/Engine/Plugins/Experimental/AnimToTexture/Source/AnimToTextureEditor/Private/AnimToTextureMeshMapping.cpp
--- a/Engine/Plugins/Experimental/AnimToTexture/Source/AnimToTextureEditor/Private/AnimToTextureMeshMapping.cpp
+++ b/Engine/Plugins/Experimental/AnimToTexture/Source/AnimToTextureEditor/Private/AnimToTextureMeshMapping.cpp
@@ -1,281 +1,195 @@
 // Copyright Epic Games, Inc. All Rights Reserved.
 #include "AnimToTextureMeshMapping.h"
+#include "Runtime/Core/Public/Async/ParallelFor.h"
 
 namespace AnimToTexture_Private
 {
 
-FVector3f FVertexToMeshMapping::FindClosestPointToTriangle(const FVector3f& P, const FVector3f& A, const FVector3f& B, const FVector3f& C)
-{
-	const FVector3f AB = B - A;
-	const FVector3f AC = C - A;
-	const FVector3f AP = P - A;
-
-	const float D1 = FVector3f::DotProduct(AB, AP);
-	const float D2 = FVector3f::DotProduct(AC, AP);
-	if (D1 <= 0.f && D2 <= 0.f)
-	{
-		return A;
-	}
-
-	const FVector3f BP = P - B;
-	const float D3 = FVector3f::DotProduct(AB, BP);
-	const float D4 = FVector3f::DotProduct(AC, BP);
-	if (D3 >= 0.f && D4 <= D3)
-	{
-		return B;
-	}
-
-	const FVector3f CP = P - C;
-	const float D5 = FVector3f::DotProduct(AB, CP);
-	const float D6 = FVector3f::DotProduct(AC, CP);
-	if (D6 >= 0.f && D5 <= D6)
-	{
-		return C;
-	}
-
-	const float VC = D1 * D4 - D3 * D2;
-	if (VC <= 0.f && D1 >= 0.f && D3 <= 0.f)
-	{
-		const float V = D1 / (D1 - D3);
-		return A + V * AB;
-	}
+FSourceVertexData::FSourceVertexData(const FVector3f& SourceVertex,
+	const TArray<FVector3f>& DriverVertices, const TArray<FIntVector3>& DriverTriangles, const TArray<VertexSkinWeightMax>& DriverSkinWeights, 
+	const float Sigma)
+{	
+	// Allocate Driver Triangle Data
+	const int32 NumDriverTriangles = DriverTriangles.Num();
 
-	const float VB = D5 * D2 - D1 * D6;
-	if (VB <= 0.f && D2 >= 0.f && D6 <= 0.f)
-	{
-		const float V = D2 / (D2 - D6);
-		return A + V * AC;
-	}
+	// Get ClosestPoint to Triangles
+	TArray<FVector3f> ClosestPoints;
+	ClosestPoints.SetNumUninitialized(NumDriverTriangles);
 
-	const float VA = D3 * D6 - D5 * D4;
-	if (VA <= 0.f && (D4 - D3) >= 0.f && (D5 - D6) >= 0.f)
+	for (int32 DriverTriangleIndex = 0; DriverTriangleIndex < NumDriverTriangles; DriverTriangleIndex++)
 	{
-		const float V = (D4 - D3) / ((D4 - D3) + (D5 - D6));
-		return B + V * (C - B);
-	}
+		const FIntVector3& DriverTriangle = DriverTriangles[DriverTriangleIndex];
+		const FVector3f& A = DriverVertices[DriverTriangle.X];
+		const FVector3f& B = DriverVertices[DriverTriangle.Y];
+		const FVector3f& C = DriverVertices[DriverTriangle.Z];
 
-	const float Denom = 1.0f / (VA + VB + VC);
-	const float V = VB * Denom;
-	const float W = VC * Denom;
+		// ClosestPoint 
+		ClosestPoints[DriverTriangleIndex] = FindClosestPointToTriangle(SourceVertex, A, B, C);
+	};
 
-	return A + V * AB + W * AC;
-}
+	// Get Inverse Distance Weighting from Vertex to each DriverTriangle (ClosestPoint)
+	TArray<float> Weights;
+	AnimToTexture_Private::InverseDistanceWeights(SourceVertex, ClosestPoints, Weights, Sigma);
 
-FVector3f FVertexToMeshMapping::GetTriangleNormal(const FVector3f& A, const FVector3f& B, const FVector3f& C)
-{
-	const FVector3f V0 = B - A;
-	const FVector3f V1 = C - A;
-	return FVector3f::CrossProduct(V0, V1).GetSafeNormal();
-}
-
-FMatrix44f FVertexToMeshMapping::GetTriangleMatrix(const FVector3f& A, const FVector3f& B, const FVector3f& C)
-{	
-	const FVector3f V0 = (B - A).GetSafeNormal();
-	const FVector3f V1 = GetTriangleNormal(A, B, C);
-	const FVector3f V2 = FVector3f::CrossProduct(V0, V1).GetSafeNormal();
-	
-	FMatrix44f Matrix = FMatrix44f::Identity; // SetAxes doesn't set [3] elements.
-	Matrix.SetAxes(&V0, &V1, &V2, &A);
-
-	return Matrix;
-}
-
-FVector3f FVertexToMeshMapping::BarycentricCoordinates(const FVector3f& P, const FVector3f& A, const FVector3f& B, const FVector3f& C)
-{
-	const FVector3f V0 = B - A;
-	const FVector3f V1 = C - A;
-	const FVector3f V2 = P - A;
-	const float D00 = FVector3f::DotProduct(V0, V0);
-	const float D01 = FVector3f::DotProduct(V0, V1);
-	const float D11 = FVector3f::DotProduct(V1, V1);
-	const float D20 = FVector3f::DotProduct(V2, V0);
-	const float D21 = FVector3f::DotProduct(V2, V1);
-	const float Denom = 1.0f / (D00 * D11 - D01 * D01);
+	// Allocate
+	DriverTriangleData.Reserve(NumDriverTriangles);
 	
-	const float V = (D11 * D20 - D01 * D21) * Denom;
-	const float W = (D00 * D21 - D01 * D20) * Denom;
-	const float U = 1.0f - V - W;
-
-	return FVector3f(U, V, W);
-}
-
-int32 FVertexToMeshMapping::FindClosestTriangle(const FVector3f& Point, const TArray<FVector3f>& Vertices, const TArray<FIntVector3>& Triangles,
-	FVector3f& OutClosestPoint, FVector3f& OutBarycentricCoords)
-{
-	int32 ClosestTriangle = INDEX_NONE;
-	float ClosestDistance = TNumericLimits<float>::Max();
-
-	for (int32 TriangleIndex = 0; TriangleIndex < Triangles.Num(); TriangleIndex++)
-	{
-		const FIntVector3& Triangle = Triangles[TriangleIndex];
-
-		// Check if valid indices
-		if (Vertices.IsValidIndex(Triangle.X) && 
-			Vertices.IsValidIndex(Triangle.Y) && 
-			Vertices.IsValidIndex(Triangle.Z))
+	for (int32 DriverTriangleIndex = 0; DriverTriangleIndex < NumDriverTriangles; DriverTriangleIndex++)
+	{	
+		// no need to store data if weight is small
+		if (Weights[DriverTriangleIndex] > UE_KINDA_SMALL_NUMBER)
 		{
-			const FVector3f& A = Vertices[Triangle.X];
-			const FVector3f& B = Vertices[Triangle.Y];
-			const FVector3f& C = Vertices[Triangle.Z];
-
-			const FVector3f ClosestPoint = FindClosestPointToTriangle(Point, A, B, C);
-			const float Distance = FVector3f::Distance(Point, ClosestPoint);
-			if (Distance < ClosestDistance)
-			{
-				OutClosestPoint = ClosestPoint;
-				ClosestTriangle = TriangleIndex;
-				ClosestDistance = Distance;
-			}
+			const FVector3f& ClosestPoint = ClosestPoints[DriverTriangleIndex];
+			const FIntVector3& DriverTriangle = DriverTriangles[DriverTriangleIndex];
+			const FVector3f& A = DriverVertices[DriverTriangle.X];
+			const FVector3f& B = DriverVertices[DriverTriangle.Y];
+			const FVector3f& C = DriverVertices[DriverTriangle.Z];
+
+			FSourceVertexDriverTriangleData TriangleData;
+			TriangleData.InverseDistanceWeight = Weights[DriverTriangleIndex];
+			TriangleData.Triangle = DriverTriangle;
+			TriangleData.BarycentricCoords = BarycentricCoordinates(ClosestPoint, A, B, C);
+			TriangleData.TangentIndex = GetTriangleTangentIndex(ClosestPoint, A, B, C);
+			TriangleData.InvMatrix = GetTriangleMatrix(ClosestPoint, A, B, C, TriangleData.TangentIndex).Inverse();
+			
+			// Interpolate SkinWeights with Barycentric Coords
+			const TArray<VertexSkinWeightMax> SkinWeights = { DriverSkinWeights[TriangleData.Triangle.X], DriverSkinWeights[TriangleData.Triangle.Y], DriverSkinWeights[TriangleData.Triangle.Z] };
+			const TArray<float> BarycentricWeights = { TriangleData.BarycentricCoords.X, TriangleData.BarycentricCoords.Y, TriangleData.BarycentricCoords.Z };
+			InterpolateVertexSkinWeights(SkinWeights, BarycentricWeights, TriangleData.SkinWeights);
+
+			DriverTriangleData.Add(TriangleData);
 		}
 	}
-	check(ClosestTriangle != INDEX_NONE);
 
-	// Compute Barycentric Coordinates
-	const FIntVector3& Triangle = Triangles[ClosestTriangle];
-	const FVector3f& A = Vertices[Triangle.X];
-	const FVector3f& B = Vertices[Triangle.Y];
-	const FVector3f& C = Vertices[Triangle.Z];
-	OutBarycentricCoords = BarycentricCoordinates(OutClosestPoint, A, B, C);
-
-	return ClosestTriangle;
 }
 
-void FVertexToMeshMapping::Create(
-	const UStaticMesh* StaticMesh, const int32 StaticMeshLODIndex,	
-	const USkeletalMesh* SkeletalMesh, const int32 SkeletalMeshLODIndex,
-	TArray<FVertexToMeshMapping>& OutMapping)
+
+FSourceMeshToDriverMesh::FSourceMeshToDriverMesh(const UStaticMesh* StaticMesh, const int32 StaticMeshLODIndex, const USkeletalMesh* SkeletalMesh, const int32 SkeletalMeshLODIndex)
 {
 	check(StaticMesh);
 	check(SkeletalMesh);
-	OutMapping.Reset();
 
 	// Get StaticMesh Vertices
-	TArray<FVector3f> StaticMeshVertices;
-	TArray<FVector3f> StaticMeshNormals;
-	const int32 NumStaticMeshVertices = GetVertices(StaticMesh, StaticMeshLODIndex, StaticMeshVertices, StaticMeshNormals);
+	const int32 NumSourceVertices = GetVertices(StaticMesh, StaticMeshLODIndex, SourceVertices, SourceNormals);
 
 	// Get SkeletalMesh Vertices
-	TArray<FVector3f> SkeletalMeshVertices;
-	const int32 NumSkeletalMeshVertices = GetVertices(SkeletalMesh, SkeletalMeshLODIndex, SkeletalMeshVertices);
+	const int32 NumDriverVertices = GetVertices(SkeletalMesh, SkeletalMeshLODIndex, DriverVertices);
 
 	// Get SkeletalMesh Triangles
-	TArray<FIntVector3> SkeletalMeshTriangles;
-	GetTriangles(SkeletalMesh, SkeletalMeshLODIndex, SkeletalMeshTriangles);
+	const int32 NumDriverTriangles = GetTriangles(SkeletalMesh, SkeletalMeshLODIndex, DriverTriangles);
 
-	// Allocate
-	OutMapping.SetNumUninitialized(NumStaticMeshVertices);
+	// Get SkeletalMesh SkinWeights
+	GetSkinWeights(SkeletalMesh, SkeletalMeshLODIndex, DriverSkinWeights);
 
-	// Fill Out Mapping
-	for (int32 VertexIndex = 0; VertexIndex < NumStaticMeshVertices; VertexIndex++)
+	// Allocate
+	SourceVerticesData.SetNumZeroed(NumSourceVertices); // note this is initializing values as zero
+	
+	// Get SourceVertex > DriverTriangle Data
+	//for (int32 SourceVertexIndex = 0; SourceVertexIndex < NumSourceVertices; SourceVertexIndex++)
+	ParallelFor(NumSourceVertices, [&](int32 SourceVertexIndex)
 	{
-		const FVector3f& StaticMeshVertex = StaticMeshVertices[VertexIndex];
+		const FVector3f& SourceVertex = SourceVertices[SourceVertexIndex];
+		const FVector3f& SourceNormal = SourceNormals[SourceVertexIndex];
 
-		FVector3f ClosestPoint;
-		FVector3f BarycentricCoords;
+		// Create Mapping from StaticMesh Vertex to SkeletalMesh
+		SourceVerticesData[SourceVertexIndex] = FSourceVertexData(SourceVertex, DriverVertices, DriverTriangles, DriverSkinWeights);
 
-		// Find Closest Triangle (on SkeletalMesh)
-		const int32 ClosestTriangle = FindClosestTriangle(StaticMeshVertex, SkeletalMeshVertices, SkeletalMeshTriangles, ClosestPoint, BarycentricCoords);
-		check(SkeletalMeshTriangles.IsValidIndex(ClosestTriangle));
-		
-		// Closest Triangle Vertices
-		const FIntVector3& Triangle = SkeletalMeshTriangles[ClosestTriangle];
-		const FVector3f& A = SkeletalMeshVertices[Triangle.X];
-		const FVector3f& B = SkeletalMeshVertices[Triangle.Y];
-		const FVector3f& C = SkeletalMeshVertices[Triangle.Z];
-		
-		// Compute Triangle Matrix
-		const FMatrix44f Matrix = GetTriangleMatrix(A, B, C);
-		
-		// Store Mapping 
-		OutMapping[VertexIndex].Triangle = Triangle;
-		OutMapping[VertexIndex].BarycentricCoords = BarycentricCoords;
-		OutMapping[VertexIndex].InvMatrix = Matrix.Inverse();
-	}
+	});	// end ParallelFor
 }
 
+int32 FSourceMeshToDriverMesh::GetNumSourceVertices() const
+{
+	return SourceVerticesData.Num();
+}
 
-void FVertexToMeshMapping::InterpolateVertexSkinWeights(const VertexSkinWeightMax& A, const VertexSkinWeightMax& B, const VertexSkinWeightMax& C,
-	VertexSkinWeightMax& OutWeights) const
-{	
-	// NOTE: this could use NumBones and remove the 256 limit.
-	TStaticArray<float, 256> DenseA(InPlace, 0.f);
-	TStaticArray<float, 256> DenseB(InPlace, 0.f);
-	TStaticArray<float, 256> DenseC(InPlace, 0.f);
+int32 FSourceMeshToDriverMesh::GetSourceVertices(TArray<FVector3f>& OutVertices) const
+{
+	OutVertices = SourceVertices;
+	return OutVertices.Num();
+}
 
-	// Convert to Dense Arrays	
-	for (int32 Index = 0; Index < MAX_TOTAL_INFLUENCES; Index++)
+int32 FSourceMeshToDriverMesh::GetSourceNormals(TArray<FVector3f>& OutNormals) const
+{
+	OutNormals = SourceNormals;
+	return OutNormals.Num();
+}
+
+void FSourceMeshToDriverMesh::DeformVerticesAndNormals(const TArray<FVector3f>& InDriverVertices,
+	TArray<FVector3f>& OutVertices, TArray<FVector3f>& OutNormals) const
+{
+	// Source Vertices
+	const int32 NumSourceVertices = SourceVerticesData.Num();
+	OutVertices.SetNumZeroed(NumSourceVertices);
+	OutNormals.SetNumZeroed(NumSourceVertices);
+
+	// Driver Triangles
+	const int32 NumDriverTriangles = DriverTriangles.Num();
+
+	// Deform Source Vertices and Normals
+	ParallelFor(NumSourceVertices, [&](int32 SourceVertexIndex)
 	{
-		if (A.BoneWeights[Index])
-		{
-			DenseA[A.MeshBoneIndices[Index]] = (float)A.BoneWeights[Index] / 255.f;
-		}
-		if (B.BoneWeights[Index])
-		{
-			DenseB[B.MeshBoneIndices[Index]] = (float)B.BoneWeights[Index] / 255.f;
-		}
-		if (C.BoneWeights[Index])
+		const FVector3f& SourceVertex = SourceVertices[SourceVertexIndex];
+		const FVector3f& SourceNormal = SourceNormals[SourceVertexIndex];
+
+		for (const FSourceVertexDriverTriangleData& TriangleData: SourceVerticesData[SourceVertexIndex].DriverTriangleData)
 		{
-			DenseC[C.MeshBoneIndices[Index]] = (float)C.BoneWeights[Index] / 255.f;
-		}
-	}
+			const FVector3f& A = InDriverVertices[TriangleData.Triangle.X];
+			const FVector3f& B = InDriverVertices[TriangleData.Triangle.Y];
+			const FVector3f& C = InDriverVertices[TriangleData.Triangle.Z];
 
-	// Blend Dense Arrays
-	TArray<TPair<uint8, uint16>> SortedVertexSkinWeights;
-	SortedVertexSkinWeights.SetNumUninitialized(256);
+			// Get Driver Triangle Point At Barycentric
+			const FVector3f& BarycentricCoords = TriangleData.BarycentricCoords;
+			const FVector3f Point = PointAtBarycentricCoordinates(A, B, C, BarycentricCoords);
 
-	for (int32 Index = 0; Index < 256; Index++)
-	{
-		const float BoneWeightA = DenseA[Index] * BarycentricCoords.X;
-		const float BoneWeightB = DenseB[Index] * BarycentricCoords.Y;
-		const float BoneWeightC = DenseC[Index] * BarycentricCoords.Z;
+			// Grt Driver Triangle Matrix
+			const uint8& TangentIndex = TriangleData.TangentIndex;
+			const FMatrix44f& DriverTriangleInvMatrix = TriangleData.InvMatrix;
+			const FMatrix44f DriverTriangleMatrix = GetTriangleMatrix(Point, A, B, C, TangentIndex);
 
-		const uint8 BoneWeight = (uint8)FMath::RoundToInt((BoneWeightA + BoneWeightB + BoneWeightC) * 255.f);
-		SortedVertexSkinWeights[Index] = TPair<uint8, uint16>(BoneWeight, Index);
-	}
+			// Get SourceVertex -> DriverTriangle InverseDistanceWeight
+			const float& InverseDistanceWeight = TriangleData.InverseDistanceWeight;
 
-	// Sort Weights and return Sparse
-	SortedVertexSkinWeights.Sort();
-	for (int32 Index = 0; Index < MAX_TOTAL_INFLUENCES; Index++)
-	{
-		OutWeights.BoneWeights[Index]     = SortedVertexSkinWeights[255 - Index].Key;
-		OutWeights.MeshBoneIndices[Index] = SortedVertexSkinWeights[255 - Index].Value;
-	}
+			// Tranform Weighted Source Vertex and Normal
+			OutVertices[SourceVertexIndex] += DriverTriangleMatrix.TransformPosition(DriverTriangleInvMatrix.TransformPosition(SourceVertex)) * InverseDistanceWeight;
+			OutNormals[SourceVertexIndex] += DriverTriangleMatrix.TransformVector(DriverTriangleInvMatrix.TransformVector(SourceNormal)) * InverseDistanceWeight;
+		}
+
+	}); // end ParallelFor
 }
 
-void FVertexToMeshMapping::InterpolateSkinWeights(const TArray<FVertexToMeshMapping>& Mapping, const TArray<VertexSkinWeightMax>& SkinWeights,
-	TArray<VertexSkinWeightMax>& OutSkinWeights)
+void FSourceMeshToDriverMesh::ProjectSkinWeights(TArray<VertexSkinWeightMax>& OutSkinWeights) const
 {
-	// Allocate StaticMesh Weights
-	OutSkinWeights.SetNumUninitialized(Mapping.Num());
+	const int32 NumSourceVertices = SourceVerticesData.Num();
 
-	for (int32 VertexIndex = 0; VertexIndex < Mapping.Num(); VertexIndex++)
+	// Allocate
+	OutSkinWeights.SetNumUninitialized(NumSourceVertices);
+
+	// Deform Source Vertices and Normals
+	ParallelFor(NumSourceVertices, [&](int32 SourceVertexIndex)
 	{
-		const FVertexToMeshMapping& VertexMapping = Mapping[VertexIndex];
-		const FIntVector3& Triangle = VertexMapping.Triangle;
+		// Allocate
+		const int32 Count = SourceVerticesData[SourceVertexIndex].DriverTriangleData.Num();
+		
+		TArray<VertexSkinWeightMax> SkinWeights;
+		SkinWeights.SetNumUninitialized(Count);
+
+		TArray<float> InverseDistanceWeights;
+		InverseDistanceWeights.SetNumUninitialized(Count);
 
-		if (SkinWeights.IsValidIndex(Triangle.X) &&
-			SkinWeights.IsValidIndex(Triangle.Y) &&
-			SkinWeights.IsValidIndex(Triangle.Z))
+		for (int32 Index = 0; Index < Count; Index++)
 		{
-			const VertexSkinWeightMax& A = SkinWeights[Triangle.X];
-			const VertexSkinWeightMax& B = SkinWeights[Triangle.Y];
-			const VertexSkinWeightMax& C = SkinWeights[Triangle.Z];
+			SkinWeights[Index] = SourceVerticesData[SourceVertexIndex].DriverTriangleData[Index].SkinWeights;
+			InverseDistanceWeights[Index] = SourceVerticesData[SourceVertexIndex].DriverTriangleData[Index].InverseDistanceWeight;
+		};
 
-			Mapping[VertexIndex].InterpolateVertexSkinWeights(A, B, C, OutSkinWeights[VertexIndex]);
-		}
-	}
-}
+		// Interpolate Driver Weights with InverseDistanceWeighting
+		InterpolateVertexSkinWeights(
+			SkinWeights,
+			InverseDistanceWeights,
+			OutSkinWeights[SourceVertexIndex]);
 
-FVector3f FVertexToMeshMapping::TransformPosition(const FVector3f& P, const FVector3f& A, const FVector3f& B, const FVector3f& C) const
-{	
-	const FMatrix44f Matrix = GetTriangleMatrix(A, B, C);	
-	return Matrix.TransformPosition(InvMatrix.TransformPosition(P));
+	}); // end ParallelFor
 }
 
-FVector3f FVertexToMeshMapping::TransformVector(const FVector3f& V, const FVector3f& A, const FVector3f& B, const FVector3f& C) const
-{
-	const FMatrix44f Matrix = GetTriangleMatrix(A, B, C);
-	return Matrix.TransformVector(InvMatrix.TransformVector(V));
-}
 
 } // end namespace AnimToTexture_Private
diff --git a/Engine/Plugins/Experimental/AnimToTexture/Source/AnimToTextureEditor/Private/AnimToTextureSkeletalMesh.cpp b/Engine/Plugins/Experimental/AnimToTexture/Source/AnimToTextureEditor/Private/AnimToTextureSkeletalMesh.cpp
--- a/Engine/Plugins/Experimental/AnimToTexture/Source/AnimToTextureEditor/Private/AnimToTextureSkeletalMesh.cpp
+++ b/Engine/Plugins/Experimental/AnimToTexture/Source/AnimToTextureEditor/Private/AnimToTextureSkeletalMesh.cpp
@@ -7,6 +7,8 @@
 #include "Components/SkeletalMeshComponent.h"
 #include "MeshDescription.h"
 #include "StaticMeshAttributes.h"
+#include "Algo/Reverse.h"
+#include "Math/NumericLimits.h"
 
 namespace AnimToTexture_Private
 {
@@ -194,6 +196,10 @@ void GetBoneNames(const USkeletalMesh* SkeletalMesh, TArray<FName>& OutNames)
 	{
 		OutNames[BoneIndex] = RefSkeleton.GetBoneName(BoneIndex);
 	}
+}
+FVector3f PointAtBarycentricCoordinates(const FVector3f& PointA, const FVector3f& PointB, const FVector3f& PointC, const FVector3f& BarycentricCoords)
+{
+	return (PointA * BarycentricCoords.X) + (PointB * BarycentricCoords.Y) + (PointC * (1.f - BarycentricCoords.X - BarycentricCoords.Y));
 };
 
 void GetSkinWeights(const USkeletalMesh* SkeletalMesh, const int32 LODIndex, 
@@ -255,6 +261,68 @@ void GetSkinWeights(const USkeletalMesh* SkeletalMesh, const int32 LODIndex,
 	}
 }
 
+void InterpolateVertexSkinWeights(const TArray<VertexSkinWeightMax>& SkinWeights, const TArray<float>& Weights,
+	VertexSkinWeightMax& OutVertexSkinWeights)
+{
+	check(SkinWeights.Num() == Weights.Num())
+
+	// Reset Values
+	OutVertexSkinWeights.BoneWeights = TStaticArray<uint8, MAX_TOTAL_INFLUENCES>(InPlace, 0.f);
+	OutVertexSkinWeights.MeshBoneIndices = TStaticArray<uint16, MAX_TOTAL_INFLUENCES>(InPlace, 0);
+
+	// Create Sparse Weighted-SkinWeights
+	TMap<uint16, float> SparseSkinWeights;
+	SparseSkinWeights.Reserve(TNumericLimits<uint16>::Max());
+
+	for (int32 Index = 0; Index < SkinWeights.Num(); Index++)
+	{
+		for (int32 InfluIndex = 0; InfluIndex < MAX_TOTAL_INFLUENCES; InfluIndex++)
+		{
+			const uint8& BoneWeight = SkinWeights[Index].BoneWeights[InfluIndex];
+			const uint16& MeshBoneIndex = SkinWeights[Index].MeshBoneIndices[InfluIndex];
+
+			const float WeightedVertexSkinWeight = (float)BoneWeight / 255.f * Weights[Index];
+
+			if (WeightedVertexSkinWeight > SMALL_NUMBER)
+			{
+				float* Value = SparseSkinWeights.Find(MeshBoneIndex);
+
+				// Inititialize
+				if (Value == nullptr)
+				{
+					SparseSkinWeights.Add(MeshBoneIndex, WeightedVertexSkinWeight);
+				}
+				// Accumulate Weighted VertexSkinWeight
+				else
+				{
+					*Value += WeightedVertexSkinWeight;
+				}
+			}
+		}
+	}
+
+	// Convert Weights to uint8 (and get them ready for Sort)
+	TArray<TPair<uint8, uint16>> SortedVertexSkinWeights;
+	SortedVertexSkinWeights.SetNumZeroed(MAX_TOTAL_INFLUENCES); // note allocate the max 12 influences
+	int32 Index = 0;
+	for (const auto& Item : SparseSkinWeights)
+	{
+		const uint8 BoneWeight = (uint8)FMath::RoundToInt(Item.Value * 255.f);
+		const uint16& MeshBoneIndex = Item.Key;
+		SortedVertexSkinWeights[Index] = TPair<uint8, uint16>(BoneWeight, MeshBoneIndex);
+		Index++;
+	}
+
+	// Sort Weights InPlace and reverse
+	SortedVertexSkinWeights.Sort();
+	Algo::Reverse(SortedVertexSkinWeights);
+	for (int32 InfluIndex = 0; InfluIndex < MAX_TOTAL_INFLUENCES; InfluIndex++)
+	{
+		OutVertexSkinWeights.BoneWeights[InfluIndex] = SortedVertexSkinWeights[InfluIndex].Key;
+		OutVertexSkinWeights.MeshBoneIndices[InfluIndex] = SortedVertexSkinWeights[InfluIndex].Value;
+	}
+}
+
 void ReduceSkinWeights(const TArray<VertexSkinWeightMax>& InSkinWeights, 
 	TArray<VertexSkinWeightFour>& OutSkinWeights)
 {
@@ -327,4 +395,172 @@ void GetSkinnedVertices(const USkeletalMeshComponent* SkeletalMeshComponent, con
 	};
 };
 
+FVector3f FindClosestPointToTriangle(const FVector3f& P, const FVector3f& A, const FVector3f& B, const FVector3f& C)
+{
+	const FVector3f AB = B - A;
+	const FVector3f AC = C - A;
+	const FVector3f AP = P - A;
+
+	const float D1 = FVector3f::DotProduct(AB, AP);
+	const float D2 = FVector3f::DotProduct(AC, AP);
+	if (D1 <= 0.f && D2 <= 0.f)
+	{
+		return A;
+	}
+
+	const FVector3f BP = P - B;
+	const float D3 = FVector3f::DotProduct(AB, BP);
+	const float D4 = FVector3f::DotProduct(AC, BP);
+	if (D3 >= 0.f && D4 <= D3)
+	{
+		return B;
+	}
+
+	const FVector3f CP = P - C;
+	const float D5 = FVector3f::DotProduct(AB, CP);
+	const float D6 = FVector3f::DotProduct(AC, CP);
+	if (D6 >= 0.f && D5 <= D6)
+	{
+		return C;
+	}
+
+	const float VC = D1 * D4 - D3 * D2;
+	if (VC <= 0.f && D1 >= 0.f && D3 <= 0.f)
+	{
+		const float V = D1 / (D1 - D3);
+		return A + V * AB;
+	}
+
+	const float VB = D5 * D2 - D1 * D6;
+	if (VB <= 0.f && D2 >= 0.f && D6 <= 0.f)
+	{
+		const float V = D2 / (D2 - D6);
+		return A + V * AC;
+	}
+
+	const float VA = D3 * D6 - D5 * D4;
+	if (VA <= 0.f && (D4 - D3) >= 0.f && (D5 - D6) >= 0.f)
+	{
+		const float V = (D4 - D3) / ((D4 - D3) + (D5 - D6));
+		return B + V * (C - B);
+	}
+
+	const float Denom = 1.0f / (VA + VB + VC);
+	const float V = VB * Denom;
+	const float W = VC * Denom;
+
+	return A + V * AB + W * AC;
+}
+
+FVector3f GetTriangleNormal(const FVector3f& A, const FVector3f& B, const FVector3f& C)
+{
+	const FVector3f V0 = B - A;
+	const FVector3f V1 = C - A;
+	return FVector3f::CrossProduct(V0, V1); // .GetSafeNormal();
+}
+
+uint8 GetTriangleTangentIndex(const FVector3f& P, const FVector3f& A, const FVector3f& B, const FVector3f& C)
+{
+	TArray<FVector3f> TangentVectors = { A - P, B - P, C - P };
+
+	float MaxDistance = 0.f;
+	uint8 TangentIndex = 0;
+	for (uint8 Index=0; Index < 3; Index++)
+	{
+		const float Distance = TangentVectors[Index].Length();
+		if (Distance > MaxDistance)
+		{
+			TangentIndex = Index;
+			MaxDistance = Distance;
+		}
+	}
+
+	return TangentIndex;
+}
+
+
+
+FMatrix44f GetTriangleMatrix(const FVector3f& P, const FVector3f& A, const FVector3f& B, const FVector3f& C, const uint8 TangentIndex)
+{
+	// Tangent Vector
+	FVector3f V0;
+	if (TangentIndex == 0)
+	{
+		V0 = A - P;
+	}
+	else if (TangentIndex == 1)
+	{
+		V0 = B - P;
+	}
+	else if (TangentIndex == 2)
+	{
+		V0 = C - P;
+	}
+
+	const FVector3f V1 = GetTriangleNormal(A, B, C);
+	const FVector3f V2 = FVector3f::CrossProduct(V0, V1);
+
+	FMatrix44f Matrix = FMatrix44f::Identity; // SetAxes doesn't set [3] elements.
+	Matrix.SetAxes(&V0, &V1, &V2, &P);
+
+	return Matrix;
+}
+
+FVector3f BarycentricCoordinates(const FVector3f& P, const FVector3f& A, const FVector3f& B, const FVector3f& C)
+{
+	const FVector3f V0 = B - A;
+	const FVector3f V1 = C - A;
+	const FVector3f V2 = P - A;
+	const float D00 = FVector3f::DotProduct(V0, V0);
+	const float D01 = FVector3f::DotProduct(V0, V1);
+	const float D11 = FVector3f::DotProduct(V1, V1);
+	const float D20 = FVector3f::DotProduct(V2, V0);
+	const float D21 = FVector3f::DotProduct(V2, V1);
+	const float Denom = 1.0f / (D00 * D11 - D01 * D01);
+
+	const float V = (D11 * D20 - D01 * D21) * Denom;
+	const float W = (D00 * D21 - D01 * D20) * Denom;
+	const float U = 1.0f - V - W;
+
+	return FVector3f(U, V, W);
+}
+
+void InverseDistanceWeights(const FVector3f& Point, const TArray<FVector3f>& Points,
+	TArray<float>& OutWeights, float Sigma)
+{
+	// Allocate
+	const int32 Count = Points.Num();
+	OutWeights.SetNumZeroed(Count);
+
+	float SumInverseDistance = 0.f;
+	TArray<float> InverseDistances;
+	InverseDistances.SetNumUninitialized(Count);
+
+	for (int32 Index = 0; Index < Count; Index++)
+	{
+		const float Distance = FVector3f::Distance(Point, Points[Index]);
+
+		// No need to interpolate if we are right on top of point.
+		if (Distance < UE_KINDA_SMALL_NUMBER)
+		{
+			OutWeights[Index] = 1.f;
+			return;
+		}
+
+		InverseDistances[Index] = 1.f / FMath::Pow(Distance, Sigma);
+		if (InverseDistances[Index] < UE_KINDA_SMALL_NUMBER)
+		{
+			InverseDistances[Index] = 0.f;
+		}
+		SumInverseDistance += InverseDistances[Index];
+	}
+
+	// Normalize Weights
+	for (int32 Index = 0; Index < Count; Index++)
+	{
+		OutWeights[Index] = InverseDistances[Index] / SumInverseDistance;
+	}
+}
+
+
 } // end namespace AnimToTexture_Private
diff --git a/Engine/Plugins/Experimental/AnimToTexture/Source/AnimToTextureEditor/Public/AnimToTextureBPLibrary.h b/Engine/Plugins/Experimental/AnimToTexture/Source/AnimToTextureEditor/Public/AnimToTextureBPLibrary.h
--- a/Engine/Plugins/Experimental/AnimToTexture/Source/AnimToTextureEditor/Public/AnimToTextureBPLibrary.h
+++ b/Engine/Plugins/Experimental/AnimToTexture/Source/AnimToTextureEditor/Public/AnimToTextureBPLibrary.h
@@ -60,7 +60,7 @@ private:
 	// Get Vertex and Normals from Current Pose
 	// The VertexDelta is returned from the RefPose
 	static void GetVertexDeltasAndNormals(const USkeletalMeshComponent* SkeletalMeshComponent, const int32 LODIndex, 
-		const TArray<FVector3f>& Vertices, const TArray<FVector3f>& Normals, const TArray<AnimToTexture_Private::FVertexToMeshMapping>& Mapping,
+		const AnimToTexture_Private::FSourceMeshToDriverMesh& SourceMeshToDriverMesh,
 		const FTransform RootTransform,
 		TArray<FVector3f>& OutVertexDeltas, TArray<FVector3f>& OutVertexNormals);
 
diff --git a/Engine/Plugins/Experimental/AnimToTexture/Source/AnimToTextureEditor/Public/AnimToTextureMeshMapping.h b/Engine/Plugins/Experimental/AnimToTexture/Source/AnimToTextureEditor/Public/AnimToTextureMeshMapping.h
--- a/Engine/Plugins/Experimental/AnimToTexture/Source/AnimToTextureEditor/Public/AnimToTextureMeshMapping.h
+++ b/Engine/Plugins/Experimental/AnimToTexture/Source/AnimToTextureEditor/Public/AnimToTextureMeshMapping.h
@@ -8,50 +8,66 @@
 namespace AnimToTexture_Private
 {
 
-// Static to Skeletal Mapping
-struct FVertexToMeshMapping
+struct FSourceVertexDriverTriangleData
 {
-	// Closest SkeletalMesh Triangle
-	FIntVector3 Triangle;
-	// SkeletalMesh Triangle Barycentric Coordinates
-	FVector3f  BarycentricCoords;
-	// SkeletalMesh Triangle Inverted Matrix
-	FMatrix44f InvMatrix;
-	
-	/* Transforms Point with given Triangle */
-	FVector3f TransformPosition(const FVector3f& Point, const FVector3f& PointA, const FVector3f& PointB, const FVector3f& PointC) const;
-	
-	/* Transforms Vector with given Triangle */
-	FVector3f TransformVector(const FVector3f& Vector, const FVector3f& PointA, const FVector3f& PointB, const FVector3f& PointC) const;
-	
-	/* Interpolate VertexSkinWeights with Barycentric Coordinates */
-	void InterpolateVertexSkinWeights(const VertexSkinWeightMax& A, const VertexSkinWeightMax& B, const VertexSkinWeightMax& C,
-		VertexSkinWeightMax& OutWeights) const;
+	uint8               TangentIndex;
+	float               InverseDistanceWeight;
+	FIntVector3         Triangle;
+	FVector3f           ClosestPoint;
+	FVector3f           BarycentricCoords;
+	FMatrix44f          InvMatrix;
+	VertexSkinWeightMax SkinWeights;
+};
+
+class FSourceVertexData
+{
+public:
+
+	FSourceVertexData() = default;	
+	FSourceVertexData(const FVector3f& SourceVertex,
+		const TArray<FVector3f>& DriverVertices, const TArray<FIntVector3>& DriverTriangles, const TArray<VertexSkinWeightMax>& DriverSkinWeights, 
+		const float Sigma=1.f);
+
+	// DriverTriangle Data specific to this SourceVertex
+	TArray<FSourceVertexDriverTriangleData> DriverTriangleData;
 	
-	/* Computes closest point to triangle*/
-	static FVector3f FindClosestPointToTriangle(const FVector3f& Point, const FVector3f& PointA, const FVector3f& PointB, const FVector3f& PointC);
+};
+
+// Creates Mapping between StaticMesh (Source) and SkeletalMesh (Driver)
+class FSourceMeshToDriverMesh
+{
+public:
+
+	FSourceMeshToDriverMesh() = default;
+	FSourceMeshToDriverMesh(const UStaticMesh* StaticMesh, const int32 StaticMeshLODIndex,
+		const USkeletalMesh* SkeletalMesh, const int32 SkeletalMeshLODIndex);
 
-	/* Computes Barycentric coordinates from point to triangle */
-	static FVector3f BarycentricCoordinates(const FVector3f& Point, const FVector3f& PointA, const FVector3f& PointB, const FVector3f& PointC);
+	// Returns Number of Source Vertices
+	int32 GetNumSourceVertices() const;
+	// Returns Source Vertices
+	int32 GetSourceVertices(TArray<FVector3f>& OutVertices) const;
+	// Returns Source Normals
+	int32 GetSourceNormals(TArray<FVector3f>& OutNormals) const;
+	
+	// Deforms Source Vertices with Driver Triangles
+	void DeformVerticesAndNormals(const TArray<FVector3f>& DriverVertices, 
+		TArray<FVector3f>& OutVertices, TArray<FVector3f>& OutNormals) const;
 
-	/* Computes Triangle Normal */
-	static FVector3f GetTriangleNormal(const FVector3f& PointA, const FVector3f& PointB, const FVector3f& PointC);
+	// Project SkinWeights
+	void ProjectSkinWeights(TArray<VertexSkinWeightMax>& OutSkinWeights) const;
 
-	/* Computes Triangle Matrix */
-	static FMatrix44f GetTriangleMatrix(const FVector3f& PointA, const FVector3f& PointB, const FVector3f& PointC);
+private:
 
-	/* Finds closest triangle, point on triangle and barycentric coordinates */
-	static int32 FindClosestTriangle(const FVector3f& Point, const TArray<FVector3f>& Vertices, const TArray<FIntVector3>& Triangles,
-		FVector3f& OutClosestPoint, FVector3f& OutBarycentricCoords);
+	// Size of Source Mesh
+	TArray<FVector3f>         SourceVertices;
+	TArray<FVector3f>         SourceNormals;
+	TArray<FSourceVertexData> SourceVerticesData;
 
-	// Creates Mapping between StaticMesh and SkeletalMesh
-	static void Create(const UStaticMesh* StaticMesh, const int32 StaticMeshLODIndex,
-		const USkeletalMesh* SkeletalMesh, const int32 SkeletalMeshLODIndex,
-		TArray<FVertexToMeshMapping>& OutMapping);
+	// Driver Data
+	TArray<FVector3f>   DriverVertices;
+	TArray<FIntVector3> DriverTriangles;
+	TArray<VertexSkinWeightMax> DriverSkinWeights;
 
-	/* Interpolate SkinWeights with VertexToMeshMapping */
-	static void InterpolateSkinWeights(const TArray<FVertexToMeshMapping>& Mapping, const TArray<VertexSkinWeightMax>& SkinWeights,
-		TArray<VertexSkinWeightMax>& OutSkinWeights);
 };
 
 } // end namespace AnimToTexture_Private
diff --git a/Engine/Plugins/Experimental/AnimToTexture/Source/AnimToTextureEditor/Public/AnimToTextureSkeletalMesh.h b/Engine/Plugins/Experimental/AnimToTexture/Source/AnimToTextureEditor/Public/AnimToTextureSkeletalMesh.h
--- a/Engine/Plugins/Experimental/AnimToTexture/Source/AnimToTextureEditor/Public/AnimToTextureSkeletalMesh.h
+++ b/Engine/Plugins/Experimental/AnimToTexture/Source/AnimToTextureEditor/Public/AnimToTextureSkeletalMesh.h
@@ -43,6 +43,10 @@ void GetSkinWeights(const USkeletalMesh* SkeletalMesh, const int32 LODIndex,
 /** Reduce Weights from MAX_TOTAL_INFLUENCES to 4 */
 void ReduceSkinWeights(const TArray<VertexSkinWeightMax>& InSkinWeights, TArray<VertexSkinWeightFour>& OutSkinWeights);
 
+/* Interpolates an Array of SkinWeights with an Array of Weights(InverseDistanceWeights) */
+void InterpolateVertexSkinWeights(const TArray<VertexSkinWeightMax>& VertexSkinWeights, const TArray<float>& Weights,
+	VertexSkinWeightMax& OutVertexSkinWeights);
+
 /* Returns Number of RawBones (no virtual bones)*/
 int32 GetNumBones(const USkeletalMesh* SkeletalMesh);
 
@@ -58,4 +62,30 @@ bool HasBone(const USkeletalMesh* SkeletalMesh, const FName& Bone);
    Only the RawBones are returned (no virtual bones) */
 void GetBoneNames(const USkeletalMesh* SkeletalMesh, TArray<FName>& OutNames);
 
+/* Computes closest point to triangle */
+FVector3f FindClosestPointToTriangle(const FVector3f& Point, const FVector3f& PointA, const FVector3f& PointB, const FVector3f& PointC);
+
+/* Computes Barycentric coordinates from point to triangle */
+FVector3f BarycentricCoordinates(const FVector3f& Point, const FVector3f& PointA, const FVector3f& PointB, const FVector3f& PointC);
+
+/* Returns Point at Barycentric coordinates of the given triangle */
+FVector3f PointAtBarycentricCoordinates(const FVector3f& PointA, const FVector3f& PointB, const FVector3f& PointC, const FVector3f& BarycentricCoords);
+
+/* Computes Triangle Normal (not normalized) */
+FVector3f GetTriangleNormal(const FVector3f& PointA, const FVector3f& PointB, const FVector3f& PointC);
+
+/* Returns TangentIndex for Triangle 
+* TangentIndex 0: PointA - PointP
+* TangentIndex 1: PointB - PointP
+* TangentIndex 2: PointC - PointP */
+uint8 GetTriangleTangentIndex(const FVector3f& Point, const FVector3f& PointA, const FVector3f& PointB, const FVector3f& PointC);
+
+/* Computes Triangle Matrix */
+FMatrix44f GetTriangleMatrix(const FVector3f& Point, const FVector3f& PointA, const FVector3f& PointB, const FVector3f& PointC, const uint8 TangentIndex);
+
+// Returns a list of weights from Point to each of the given Points.
+// The returned Weights are the size of Points and normalized (sum to 1.0) 
+void InverseDistanceWeights(const FVector3f& Point, const TArray<FVector3f>& Points,
+	TArray<float>& OutWeights, float Sigma = 1.f);
+
 } // end namespace AnimToTexture_Private
