[AnimToTexture]
- Added Progress Bar
- Moved DataAsset Checks to function

diff --git a/Engine/Plugins/Experimental/AnimToTexture/Source/AnimToTexture/Public/AnimToTextureDataAsset.h b/Engine/Plugins/Experimental/AnimToTexture/Source/AnimToTexture/Public/AnimToTextureDataAsset.h
--- a/Engine/Plugins/Experimental/AnimToTexture/Source/AnimToTexture/Public/AnimToTextureDataAsset.h
+++ b/Engine/Plugins/Experimental/AnimToTexture/Source/AnimToTexture/Public/AnimToTextureDataAsset.h
@@ -3,7 +3,7 @@
 #pragma once
 
 #include "Engine/DataAsset.h"
-#include "Animation/AnimInstance.h"
+#include "Animation/AnimSequence.h"
 #include "Engine/SkeletalMesh.h"
 #include "Engine/StaticMesh.h"
 #include "Engine/Texture2D.h"
diff --git a/Engine/Plugins/Experimental/AnimToTexture/Source/AnimToTextureEditor/Private/AnimToTextureBPLibrary.cpp b/Engine/Plugins/Experimental/AnimToTexture/Source/AnimToTextureEditor/Private/AnimToTextureBPLibrary.cpp
--- a/Engine/Plugins/Experimental/AnimToTexture/Source/AnimToTextureEditor/Private/AnimToTextureBPLibrary.cpp
+++ b/Engine/Plugins/Experimental/AnimToTexture/Source/AnimToTextureEditor/Private/AnimToTextureBPLibrary.cpp
@@ -18,6 +18,9 @@
 #include "MeshDescription.h"
 #include "Materials/MaterialInstanceConstant.h"
 #include "MaterialEditingLibrary.h"
+#include "Misc/ScopedSlowTask.h"
+
+#define LOCTEXT_NAMESPACE "UAnimToTextureBPLibrary"
 
 using namespace AnimToTexture_Private;
 
@@ -34,108 +37,30 @@ bool UAnimToTextureBPLibrary::AnimationToTexture(UAnimToTextureDataAsset* DataAs
 		return false;
 	}
 
-	// Reset DataAsset Info Values
-	DataAsset->ResetInfo();
-
-	// Check StaticMesh
-	if (!DataAsset->GetStaticMesh())
-	{
-		UE_LOG(LogAnimToTextureEditor, Warning, TEXT("Invalid StaticMesh"));
-		return false;
-	}
-
-	// Check SkeletalMesh
-	if (!DataAsset->GetSkeletalMesh())
-	{
-		UE_LOG(LogAnimToTextureEditor, Warning, TEXT("Invalid SkeletalMesh"));
-		return false;
-	}
-
-	// Check Skeleton
-	if (!DataAsset->GetSkeletalMesh()->GetSkeleton())
-	{
-		UE_LOG(LogAnimToTextureEditor, Warning, TEXT("Invalid SkeletalMesh. No valid Skeleton found"));
-		return false;
-	}
-
-	// Check StaticMesh LOD
-	if (!DataAsset->GetStaticMesh()->IsSourceModelValid(DataAsset->StaticLODIndex))
-	{
-		UE_LOG(LogAnimToTextureEditor, Warning, TEXT("Invalid StaticMesh LOD Index: %i"), DataAsset->StaticLODIndex);
-		return false;
-	}
-
-	// Check SkeletalMesh LOD
-	if (!DataAsset->GetSkeletalMesh()->IsValidLODIndex(DataAsset->SkeletalLODIndex))
-	{
-		UE_LOG(LogAnimToTextureEditor, Warning, TEXT("Invalid SkeletalMesh LOD Index: %i"), DataAsset->SkeletalLODIndex);
-		return false;
-	}
-
-	// Check Socket.
-	bool bValidSocket = false;
-	if (DataAsset->AttachToSocket.IsValid() && !DataAsset->AttachToSocket.IsNone())
-	{
-		if (HasBone(DataAsset->GetSkeletalMesh(), DataAsset->AttachToSocket))
-		{
-			bValidSocket = true;
-		}
-		else
-		{
-			UE_LOG(LogAnimToTextureEditor, Warning, TEXT("Invalid Socket: %s"), *DataAsset->AttachToSocket.ToString());
-			return false;
-		}
-	}
-	if (bValidSocket && DataAsset->Mode == EAnimToTextureMode::Vertex)
+	// Runs some checks for the assets in DataAsset
+	int32 SocketIndex = INDEX_NONE;
+	TArray<FAnimToTextureAnimSequenceInfo> AnimSequences;
+	if (!CheckDataAsset(DataAsset, SocketIndex, AnimSequences))
 	{
-		UE_LOG(LogAnimToTextureEditor, Warning, TEXT("Unable to use Socket in Vertex Mode. Use Bone Mode instead."));
 		return false;
 	}
 
-	// Check if UVChannel is being used by the Lightmap UV
-	const FStaticMeshSourceModel& SourceModel = DataAsset->GetStaticMesh()->GetSourceModel(DataAsset->StaticLODIndex);
-	if (SourceModel.BuildSettings.bGenerateLightmapUVs &&
-		SourceModel.BuildSettings.DstLightmapIndex == DataAsset->UVChannel)
-	{
-		UE_LOG(LogAnimToTextureEditor, Warning, TEXT("Invalid UVChannel: %i. Already used by LightMap"), DataAsset->UVChannel);
-		return false;
-	}
+	// Reset DataAsset Info Values
+	DataAsset->ResetInfo();
 
-	// Check Animations
-	int32 NumAnimations = 0;
-	for (const FAnimToTextureAnimSequenceInfo& AnimSequenceInfo : DataAsset->AnimSequences)
+	// ---------------------------------------------------------------------------		
+	// Get Mapping between Static and Skeletal Meshes
+	// Since they might not have same number of points.
+	//
+	FSourceMeshToDriverMesh Mapping;
 	{
-		const UAnimSequence* AnimSequence = AnimSequenceInfo.AnimSequence;
-		if (AnimSequenceInfo.bEnabled && AnimSequence)
-		{
-			// Check Frame Range
-			if (AnimSequenceInfo.bUseCustomRange &&
-				(AnimSequenceInfo.StartFrame < 0 ||
-				 AnimSequenceInfo.EndFrame > AnimSequence->GetNumberOfSampledKeys() - 1 ||
-				 AnimSequenceInfo.EndFrame - AnimSequenceInfo.StartFrame < 0 ))
-			{
-				UE_LOG(LogAnimToTextureEditor, Warning, TEXT("Invalid Custom Range for AnimSequence: %s"), *AnimSequence->GetName());
-				return false;
-			}
+		FScopedSlowTask ProgressBar(1.f, LOCTEXT("ProcessingAnimSequence", "Processing StaticMesh -> SkeletalMesh Mapping ..."), true /*Enabled*/);
+		ProgressBar.MakeDialog(false /*bShowCancelButton*/, false /*bAllowInPIE*/);
 
-			NumAnimations++;
-		}
+		Mapping.Update(DataAsset->GetStaticMesh(), DataAsset->StaticLODIndex,
+			DataAsset->GetSkeletalMesh(), DataAsset->SkeletalLODIndex);
 	}
-	if (!NumAnimations)
-	{
-		UE_LOG(LogAnimToTextureEditor, Warning, TEXT("No Animations found"));
-		return false;
-	};
-
-	// ---------------------------------------------------------------------------		
-	// Get Meshes Vertices and Mapping.
-	// NOTE: We need to create a Mapping between the StaticMesh and the SkeletalMesh
-	//       Since they dont have same number of points.
-	//
 
-	// Get SourceMeshToDriverMesh
-	FSourceMeshToDriverMesh Mapping(DataAsset->GetStaticMesh(), DataAsset->StaticLODIndex, DataAsset->GetSkeletalMesh(), DataAsset->SkeletalLODIndex);
-	
 	// Get Number of Source Vertices (StaticMesh)
 	const int32 NumVertices = Mapping.GetNumSourceVertices();
 	if (!NumVertices)
@@ -147,8 +72,6 @@ bool UAnimToTextureBPLibrary::AnimationToTexture(UAnimToTextureDataAsset* DataAs
 	// Get Reference Skeleton Transforms
 	//
 	int32 NumBones = INDEX_NONE;
-	int32 SocketIndex = INDEX_NONE;
-	TArray<FName>     BoneNames;
 	TArray<FVector3f> BoneRefPositions;
 	TArray<FVector4>  BoneRefRotations;
 	TArray<FVector3f> BonePositions;
@@ -160,28 +83,6 @@ bool UAnimToTextureBPLibrary::AnimationToTexture(UAnimToTextureDataAsset* DataAs
 		NumBones = GetRefBonePositionsAndRotations(DataAsset->GetSkeletalMesh(),
 			BoneRefPositions, BoneRefRotations);
 
-		// NOTE: there is a limitation with the number of bones atm.
-		if (NumBones > 256)
-		{
-			UE_LOG(LogAnimToTextureEditor, Warning, TEXT("Invalid Number of Bones. There is a maximum of 256 bones"))
-			return false;
-		}
-
-		// Get Bone Names (no virtual)
-		GetBoneNames(DataAsset->GetSkeletalMesh(), BoneNames);
-
-		// Make sure array sizes are correct.
-		check(BoneNames.Num() == NumBones);
-
-		// Check if Socket is in BoneNames
-		if (bValidSocket && !BoneNames.Find(DataAsset->AttachToSocket, SocketIndex))
-		{
-			UE_LOG(LogAnimToTextureEditor, Warning, TEXT("Socket: %s not found in Raw Bone List"), *DataAsset->AttachToSocket.ToString());
-			return false;
-		}
-
-		// TODO: SocketIndex can only be < TNumericLimits<uint16>::Max()
-		
 		// Add RefPose 
 		// Note: this is added in the first frame of the Bone Position and Rotation Textures
 		BonePositions.Append(BoneRefPositions);
@@ -216,66 +117,42 @@ bool UAnimToTextureBPLibrary::AnimationToTexture(UAnimToTextureDataAsset* DataAs
 	
 	// Get Animation Frames Data
 	//
-	for (const FAnimToTextureAnimSequenceInfo& AnimSequenceInfo : DataAsset->AnimSequences)
+	for (int32 AnimSequenceIndex = 0; AnimSequenceIndex < AnimSequences.Num(); AnimSequenceIndex++)
 	{
-		UAnimSequence* AnimSequence = AnimSequenceInfo.AnimSequence;
+		const FAnimToTextureAnimSequenceInfo& AnimSequenceInfo = AnimSequences[AnimSequenceIndex];
 
-		if (!AnimSequenceInfo.bEnabled || !AnimSequence)
-		{
-			continue;
-		}
-		
-		// Make sure SkeletalMesh is compatible with AnimSequence
-		// TODO: move this to early Checks.
-		if (!SkeletalMeshComponent->GetSkeletalMeshAsset()->GetSkeleton()->IsCompatibleForEditor(AnimSequence->GetSkeleton()))
-		{
-			UE_LOG(LogAnimToTextureEditor, Warning, TEXT("Invalid AnimSequence: %s for given SkeletalMesh: %s"), *AnimSequence->GetFName().ToString(), *SkeletalMeshComponent->GetSkeletalMeshAsset()->GetFName().ToString());
-			continue;
-		}
-		// Set AnimSequence
-		else
-		{
-			SkeletalMeshComponent->SetAnimation(AnimSequence);
-		}
+		// Set Animation
+		UAnimSequence* AnimSequence = AnimSequenceInfo.AnimSequence;
+		SkeletalMeshComponent->SetAnimation(AnimSequence);
 
-		// -----------------------------------------------------------------------------------
 		// Get Number of Frames
-		//
 		int32 AnimStartFrame;
 		int32 AnimEndFrame;
-		
-		// Get Range from AnimSequence
-		if (!AnimSequenceInfo.bUseCustomRange)
-		{
-			AnimStartFrame = 0;
-			AnimEndFrame = AnimSequence->GetNumberOfSampledKeys() - 1; // AnimSequence->GetNumberOfFrames();
-		}
-		// Get Range from DataAsset
-		else
-		{
-			AnimStartFrame = AnimSequenceInfo.StartFrame;
-			AnimEndFrame = AnimSequenceInfo.EndFrame;
-		}
-		
-		// ---------------------------------------------------------------------------
-		// 
-		const int32 AnimNumFrames = AnimEndFrame - AnimStartFrame + 1;
+		const int32 AnimNumFrames = GetAnimationFrameRange(AnimSequenceInfo, AnimStartFrame, AnimEndFrame);
 		const float AnimStartTime = AnimSequence->GetTimeAtFrame(AnimStartFrame);
 
 		int32 SampleIndex = 0;
 		const float SampleInterval = 1.f / DataAsset->SampleRate;
 
+		// Progress Bar
+		FFormatNamedArguments Args;
+		Args.Add(TEXT("AnimSequenceIndex"), AnimSequenceIndex+1);
+		Args.Add(TEXT("NumAnimSequences"), AnimSequences.Num());
+		Args.Add(TEXT("AnimSequence"), FText::FromString(*AnimSequence->GetFName().ToString()));
+		FScopedSlowTask AnimProgressBar(AnimNumFrames, FText::Format(LOCTEXT("ProcessingAnimSequence", "Processing AnimSequence: {AnimSequence} [{AnimSequenceIndex}/{NumAnimSequences}]"), Args), true /*Enabled*/);
+		AnimProgressBar.MakeDialog(false /*bShowCancelButton*/, false /*bAllowInPIE*/);
+
 		while (SampleIndex < AnimNumFrames)
 		{
+			AnimProgressBar.EnterProgressFrame();
+
 			const float Time = AnimStartTime + ((float)SampleIndex * SampleInterval);
 			SampleIndex++;
 
 			// Go To Time
 			SkeletalMeshComponent->SetPosition(Time);
-
 			// Update SkelMesh Animation.
 			SkeletalMeshComponent->TickAnimation(0.f, false /*bNeedsValidRootMotion*/);
-			// SkeletalMeshComponent->TickComponent(0.f, ELevelTick::LEVELTICK_All, nullptr);
 			SkeletalMeshComponent->RefreshBoneTransforms(nullptr /*TickFunction*/);
 			
 			// ---------------------------------------------------------------------------
@@ -293,7 +170,7 @@ bool UAnimToTextureBPLibrary::AnimationToTexture(UAnimToTextureDataAsset* DataAs
 				VertexDeltas.Append(VertexFrameDeltas);
 				VertexNormals.Append(VertexFrameNormals);
 
-			} // End Vertex Mode
+			}
 
 			// ---------------------------------------------------------------------------
 			// Store Bone Positions & Rotations
@@ -309,8 +186,7 @@ bool UAnimToTextureBPLibrary::AnimationToTexture(UAnimToTextureDataAsset* DataAs
 				BonePositions.Append(BoneFramePositions);
 				BoneRotations.Append(BoneFrameRotations);
 
-			} // End Bone Mode
-
+			}
 		} // End Frame
 
 		// Store Anim Info Data
@@ -321,6 +197,7 @@ bool UAnimToTextureBPLibrary::AnimationToTexture(UAnimToTextureDataAsset* DataAs
 
 		// Accumulate Frames
 		DataAsset->NumFrames += AnimNumFrames;
+
 	} // End Anim
 		
 	// Destroy Temp Component & Actor
@@ -329,14 +206,7 @@ bool UAnimToTextureBPLibrary::AnimationToTexture(UAnimToTextureDataAsset* DataAs
 	Actor->Destroy();
 	
 	// ---------------------------------------------------------------------------
-	// Nothing to do here ...
-	//
-	if (!DataAsset->NumFrames)
-	{
-		return false;
-	}
 
-	// ---------------------------------------------------------------------------
 	if (DataAsset->Mode == EAnimToTextureMode::Vertex)
 	{
 		// Find Best Resolution for Vertex Data
@@ -386,17 +256,17 @@ bool UAnimToTextureBPLibrary::AnimationToTexture(UAnimToTextureDataAsset* DataAs
 		// Find Best Resolution for Bone Data
 		int32 Height, Width;
 
-		// Note we are adding +1 frame for the ref pose
-		if (!FindBestResolution(DataAsset->NumFrames + 1, NumBones, 
-								Height, Width, DataAsset->BoneRowsPerFrame, 
-								DataAsset->MaxHeight, DataAsset->MaxWidth, DataAsset->bEnforcePowerOfTwo))
-		{
-			UE_LOG(LogAnimToTextureEditor, Warning, TEXT("Bone Animation data cannot be fit in a %ix%i texture."), DataAsset->MaxHeight, DataAsset->MaxWidth);
-			return false;
-		}
-
 		// Write Bone Position and Rotation Textures
 		{
+			// Note we are adding +1 frame for the ref pose
+			if (!FindBestResolution(DataAsset->NumFrames + 1, NumBones,
+				Height, Width, DataAsset->BoneRowsPerFrame,
+				DataAsset->MaxHeight, DataAsset->MaxWidth, DataAsset->bEnforcePowerOfTwo))
+			{
+				UE_LOG(LogAnimToTextureEditor, Warning, TEXT("Bone Animation data cannot be fit in a %ix%i texture."), DataAsset->MaxHeight, DataAsset->MaxWidth);
+				return false;
+			}
+
 			// Normalize Bone Data
 			TArray<FVector3f> NormalizedBonePositions;
 			TArray<FVector4> NormalizedBoneRotations;
@@ -416,25 +286,28 @@ bool UAnimToTextureBPLibrary::AnimationToTexture(UAnimToTextureDataAsset* DataAs
 				WriteVectorsToTexture<FVector3f, FLowPrecision>(NormalizedBonePositions, DataAsset->NumFrames + 1, DataAsset->BoneRowsPerFrame, Height, Width, DataAsset->GetBonePositionTexture());
 				WriteVectorsToTexture<FVector4, FLowPrecision>(NormalizedBoneRotations, DataAsset->NumFrames + 1, DataAsset->BoneRowsPerFrame, Height, Width, DataAsset->GetBoneRotationTexture());
 			}
-		}
 
-		// ---------------------------------------------------------------------------
-
-		// Find Best Resolution for Bone Weights Texture
-		if (!FindBestResolution(2, NumVertices, 
-								Height, Width, DataAsset->BoneWeightRowsPerFrame, 
-								DataAsset->MaxHeight, DataAsset->MaxWidth, DataAsset->bEnforcePowerOfTwo))
-		{
-			UE_LOG(LogAnimToTextureEditor, Warning, TEXT("Weights Data cannot be fit in a %ix%i texture."), DataAsset->MaxHeight, DataAsset->MaxWidth);
-			return false;
+			// Update Bounds
+			SetBoundsExtensions(DataAsset->GetStaticMesh(), DataAsset->BoneMinBBox, DataAsset->BoneSizeBBox);
 		}
 
+		// ---------------------------------------------------------------------------
+		
 		// Write Weights Texture
 		{
+			// Find Best Resolution for Bone Weights Texture
+			if (!FindBestResolution(2, NumVertices,
+				Height, Width, DataAsset->BoneWeightRowsPerFrame,
+				DataAsset->MaxHeight, DataAsset->MaxWidth, DataAsset->bEnforcePowerOfTwo))
+			{
+				UE_LOG(LogAnimToTextureEditor, Warning, TEXT("Weights Data cannot be fit in a %ix%i texture."), DataAsset->MaxHeight, DataAsset->MaxWidth);
+				return false;
+			}
+
 			TArray<TVertexSkinWeight<4>> SkinWeights;
 
 			// Reduce BoneWeights to 4 Influences.
-			if (!bValidSocket)
+			if (SocketIndex == INDEX_NONE)
 			{	
 				// Project SkinWeights from SkeletalMesh to StaticMesh
 				TArray<VertexSkinWeightMax> StaticMeshSkinWeights;
@@ -443,7 +316,6 @@ bool UAnimToTextureBPLibrary::AnimationToTexture(UAnimToTextureDataAsset* DataAs
 				// Reduce Weights to 4 highest influences.
 				ReduceSkinWeights(StaticMeshSkinWeights, SkinWeights);
 			}
-
 			// If Valid Socket, set all influences to same index.
 			else
 			{
@@ -459,13 +331,10 @@ bool UAnimToTextureBPLibrary::AnimationToTexture(UAnimToTextureDataAsset* DataAs
 			// Write Bone Weights Texture
 			WriteSkinWeightsToTexture(SkinWeights,
 				DataAsset->BoneWeightRowsPerFrame, Height, Width, DataAsset->GetBoneWeightTexture());
-		}
 
-		// Add Vertex UVChannel
-		CreateUVChannel(DataAsset->GetStaticMesh(), DataAsset->StaticLODIndex, DataAsset->UVChannel, Height, Width);
-
-		// Update Bounds
-		SetBoundsExtensions(DataAsset->GetStaticMesh(), DataAsset->BoneMinBBox, DataAsset->BoneSizeBBox);
+			// Add Vertex UVChannel
+			CreateUVChannel(DataAsset->GetStaticMesh(), DataAsset->StaticLODIndex, DataAsset->UVChannel, Height, Width);
+		}
 
 		// Done with StaticMesh
 		DataAsset->GetStaticMesh()->PostEditChange();
@@ -481,6 +350,140 @@ bool UAnimToTextureBPLibrary::AnimationToTexture(UAnimToTextureDataAsset* DataAs
 }
 
 
+bool UAnimToTextureBPLibrary::CheckDataAsset(const UAnimToTextureDataAsset* DataAsset, 
+	int32& OutSocketIndex, TArray<FAnimToTextureAnimSequenceInfo>& OutAnimSequences)
+{
+	// Check StaticMesh
+	if (!DataAsset->GetStaticMesh())
+	{
+		UE_LOG(LogAnimToTextureEditor, Warning, TEXT("Invalid StaticMesh"));
+		return false;
+	}
+
+	// Check SkeletalMesh
+	if (!DataAsset->GetSkeletalMesh())
+	{
+		UE_LOG(LogAnimToTextureEditor, Warning, TEXT("Invalid SkeletalMesh"));
+		return false;
+	}
+
+	// Check Skeleton
+	if (!DataAsset->GetSkeletalMesh()->GetSkeleton())
+	{
+		UE_LOG(LogAnimToTextureEditor, Warning, TEXT("Invalid SkeletalMesh. No valid Skeleton found"));
+		return false;
+	}
+
+	// Check StaticMesh LOD
+	if (!DataAsset->GetStaticMesh()->IsSourceModelValid(DataAsset->StaticLODIndex))
+	{
+		UE_LOG(LogAnimToTextureEditor, Warning, TEXT("Invalid StaticMesh LOD Index: %i"), DataAsset->StaticLODIndex);
+		return false;
+	}
+
+	// Check SkeletalMesh LOD
+	if (!DataAsset->GetSkeletalMesh()->IsValidLODIndex(DataAsset->SkeletalLODIndex))
+	{
+		UE_LOG(LogAnimToTextureEditor, Warning, TEXT("Invalid SkeletalMesh LOD Index: %i"), DataAsset->SkeletalLODIndex);
+		return false;
+	}
+
+	// Check Socket.
+	OutSocketIndex = INDEX_NONE;
+	if (DataAsset->AttachToSocket.IsValid() && !DataAsset->AttachToSocket.IsNone())
+	{
+		// Get Bone Names (no virtual)
+		TArray<FName> BoneNames;
+		GetBoneNames(DataAsset->GetSkeletalMesh(), BoneNames);
+
+		// Check if Socket is in BoneNames
+		if (!BoneNames.Find(DataAsset->AttachToSocket, OutSocketIndex))
+		{
+			UE_LOG(LogAnimToTextureEditor, Warning, TEXT("Socket: %s not found in Raw Bone List"), *DataAsset->AttachToSocket.ToString());
+			return false;
+		}
+		else
+		{
+			// TODO: SocketIndex can only be < TNumericLimits<uint16>::Max()
+		}
+	}
+
+	// Check if UVChannel is being used by the Lightmap UV
+	const FStaticMeshSourceModel& SourceModel = DataAsset->GetStaticMesh()->GetSourceModel(DataAsset->StaticLODIndex);
+	if (SourceModel.BuildSettings.bGenerateLightmapUVs &&
+		SourceModel.BuildSettings.DstLightmapIndex == DataAsset->UVChannel)
+	{
+		UE_LOG(LogAnimToTextureEditor, Warning, TEXT("Invalid UVChannel: %i. Already used by LightMap"), DataAsset->UVChannel);
+		return false;
+	}
+
+	// Check if NumBones > 256
+	const int32 NumBones = GetNumBones(DataAsset->GetSkeletalMesh());
+	if (DataAsset->Precision == EAnimToTexturePrecision::EightBits &&
+		NumBones > 256)
+	{
+		UE_LOG(LogAnimToTextureEditor, Warning, TEXT("Invalid Number of Bones. There is a maximum of 256 bones for 8bit Precision"))
+		return false;
+	}
+
+	// Check Animations
+	OutAnimSequences.Reset();
+	for (const FAnimToTextureAnimSequenceInfo& AnimSequenceInfo : DataAsset->AnimSequences)
+	{
+		const UAnimSequence* AnimSequence = AnimSequenceInfo.AnimSequence;
+
+		if (AnimSequenceInfo.bEnabled && AnimSequence)
+		{
+			// Make sure SkeletalMesh is compatible with AnimSequence
+			if (!DataAsset->GetSkeletalMesh()->GetSkeleton()->IsCompatibleForEditor(AnimSequence->GetSkeleton()))
+			{
+				UE_LOG(LogAnimToTextureEditor, Warning, TEXT("Invalid AnimSequence: %s for given SkeletalMesh: %s"), *AnimSequence->GetFName().ToString(), *DataAsset->GetSkeletalMesh()->GetFName().ToString());
+				return false;
+			}
+
+			// Check Frame Range
+			if (AnimSequenceInfo.bUseCustomRange &&
+				(AnimSequenceInfo.StartFrame < 0 ||
+					AnimSequenceInfo.EndFrame > AnimSequence->GetNumberOfSampledKeys() - 1 ||
+					AnimSequenceInfo.EndFrame - AnimSequenceInfo.StartFrame < 0))
+			{
+				UE_LOG(LogAnimToTextureEditor, Warning, TEXT("Invalid CustomRange for AnimSequence: %s"), *AnimSequence->GetName());
+				return false;
+			}
+
+			// Store Valid AnimSequenceInfo
+			OutAnimSequences.Add(AnimSequenceInfo);
+		}
+	}
+
+	// All Good !
+	return true;
+}
+
+int32 UAnimToTextureBPLibrary::GetAnimationFrameRange(const FAnimToTextureAnimSequenceInfo& Animation, int32& OutStartFrame, int32& OutEndFrame)
+{
+	if (!Animation.AnimSequence)
+	{
+		return INDEX_NONE;
+	}
+
+	// Get Range from AnimSequence
+	if (!Animation.bUseCustomRange)
+	{
+		OutStartFrame = 0;
+		OutEndFrame = Animation.AnimSequence->GetNumberOfSampledKeys() - 1; // AnimSequence->GetNumberOfFrames();
+	}
+	// Get Custom Range
+	else
+	{
+		OutStartFrame = Animation.StartFrame;
+		OutEndFrame = Animation.EndFrame;
+	}
+
+	// Return Number of Frames
+	return OutEndFrame - OutStartFrame + 1;
+}
+
 // 
 void UAnimToTextureBPLibrary::GetVertexDeltasAndNormals(const USkeletalMeshComponent* SkeletalMeshComponent, const int32 LODIndex, 
 	const AnimToTexture_Private::FSourceMeshToDriverMesh& SourceMeshToDriverMesh,
@@ -528,14 +531,11 @@ void UAnimToTextureBPLibrary::GetVertexDeltasAndNormals(const USkeletalMeshCompo
 int32 UAnimToTextureBPLibrary::GetRefBonePositionsAndRotations(const USkeletalMesh* SkeletalMesh,
 	TArray<FVector3f>& OutBoneRefPositions, TArray<FVector4>& OutBoneRefRotations)
 {
+	check(SkeletalMesh);
+
 	OutBoneRefPositions.Reset();
 	OutBoneRefRotations.Reset();
 
-	if (!SkeletalMesh)
-	{
-		return 0;
-	}
-
 	// Get Number of RawBones (no virtual)
 	const int32 NumBones = GetNumBones(SkeletalMesh);
 	
@@ -551,6 +551,8 @@ int32 UAnimToTextureBPLibrary::GetRefBonePositionsAndRotations(const USkeletalMe
 int32 UAnimToTextureBPLibrary::GetBonePositionsAndRotations(const USkeletalMeshComponent* SkeletalMeshComponent, const TArray<FVector3f>& BoneRefPositions,
 	TArray<FVector3f>& BonePositions, TArray<FVector4>& BoneRotations)
 {
+	check(SkeletalMeshComponent);
+
 	BonePositions.Reset();
 	BoneRotations.Reset();
 
@@ -602,11 +604,9 @@ void UAnimToTextureBPLibrary::UpdateMaterialInstanceFromDataAsset(UAnimToTexture
 	const bool bAutoPlay, const int32 AnimationIndex, 
 	const EAnimToTextureNumBoneInfluences NumBoneInfluences, const EMaterialParameterAssociation MaterialParameterAssociation)
 {
-	if (!MaterialInstance || !DataAsset)
-	{
-		return;
-	}
-
+	check(DataAsset);
+	check(MaterialInstance);
+	
 	// Set UVChannel
 	switch (DataAsset->UVChannel)
 	{
@@ -725,12 +725,9 @@ void UAnimToTextureBPLibrary::UpdateMaterialInstanceFromDataAsset(UAnimToTexture
 
 bool UAnimToTextureBPLibrary::SetLightMapIndex(UStaticMesh* StaticMesh, const int32 LODIndex, const int32 LightmapIndex, bool bGenerateLightmapUVs)
 {
-	if (!StaticMesh)
-	{
-		return false;
-	}
+	check(StaticMesh);
 
-	if (LODIndex >= 0 && !StaticMesh->IsSourceModelValid(LODIndex))
+	if (!StaticMesh->IsSourceModelValid(LODIndex))
 	{
 		return false;
 	}
@@ -760,17 +757,14 @@ bool UAnimToTextureBPLibrary::SetLightMapIndex(UStaticMesh* StaticMesh, const in
 
 UStaticMesh* UAnimToTextureBPLibrary::ConvertSkeletalMeshToStaticMesh(USkeletalMesh* SkeletalMesh, const FString PackageName, const int32 LODIndex)
 {
-	if (!SkeletalMesh || PackageName.IsEmpty())
-	{
-		return nullptr;
-	}
+	check(SkeletalMesh);
 
-	if (!FPackageName::IsValidObjectPath(PackageName))
+	if (PackageName.IsEmpty() || !FPackageName::IsValidObjectPath(PackageName))
 	{
 		return nullptr;
 	}
 
-	if (LODIndex >= 0 && !SkeletalMesh->IsValidLODIndex(LODIndex))
+	if (!SkeletalMesh->IsValidLODIndex(LODIndex))
 	{
 		UE_LOG(LogAnimToTextureEditor, Warning, TEXT("Invalid LODIndex: %i"), LODIndex);
 		return nullptr;
@@ -1135,3 +1129,5 @@ void UAnimToTextureBPLibrary::SetBoundsExtensions(UStaticMesh* StaticMesh, const
 	StaticMesh->SetNegativeBoundsExtension(NewNegativeBoundsExtension);
 	StaticMesh->CalculateExtendedBounds();
 }
+
+#undef LOCTEXT_NAMESPACE
diff --git a/Engine/Plugins/Experimental/AnimToTexture/Source/AnimToTextureEditor/Private/AnimToTextureMeshMapping.cpp b/Engine/Plugins/Experimental/AnimToTexture/Source/AnimToTextureEditor/Private/AnimToTextureMeshMapping.cpp
--- a/Engine/Plugins/Experimental/AnimToTexture/Source/AnimToTextureEditor/Private/AnimToTextureMeshMapping.cpp
+++ b/Engine/Plugins/Experimental/AnimToTexture/Source/AnimToTextureEditor/Private/AnimToTextureMeshMapping.cpp
@@ -5,12 +5,13 @@
 namespace AnimToTexture_Private
 {
 
-FSourceVertexData::FSourceVertexData(const FVector3f& SourceVertex,
+void FSourceVertexData::Update(const FVector3f& SourceVertex,
 	const TArray<FVector3f>& DriverVertices, const TArray<FIntVector3>& DriverTriangles, const TArray<VertexSkinWeightMax>& DriverSkinWeights, 
 	const float Sigma)
 {	
-	// Allocate Driver Triangle Data
+	// Reserve TriangleData
 	const int32 NumDriverTriangles = DriverTriangles.Num();
+	DriverTriangleData.Reserve(NumDriverTriangles);
 
 	// Get ClosestPoint to Triangles
 	TArray<FVector3f> ClosestPoints;
@@ -23,19 +24,17 @@ FSourceVertexData::FSourceVertexData(const FVector3f& SourceVertex,
 		const FVector3f& B = DriverVertices[DriverTriangle.Y];
 		const FVector3f& C = DriverVertices[DriverTriangle.Z];
 
-		// ClosestPoint 
-		ClosestPoints[DriverTriangleIndex] = FindClosestPointToTriangle(SourceVertex, A, B, C);
+		// ClosestPoint To Triangle 
+		int32 OnPointLocalIndex = INDEX_NONE;
+		ClosestPoints[DriverTriangleIndex] = FindClosestPointToTriangle(SourceVertex, A, B, C, OnPointLocalIndex);
 	};
 
 	// Get Inverse Distance Weighting from Vertex to each DriverTriangle (ClosestPoint)
 	TArray<float> Weights;
 	AnimToTexture_Private::InverseDistanceWeights(SourceVertex, ClosestPoints, Weights, Sigma);
-
-	// Allocate
-	DriverTriangleData.Reserve(NumDriverTriangles);
 	
 	for (int32 DriverTriangleIndex = 0; DriverTriangleIndex < NumDriverTriangles; DriverTriangleIndex++)
-	{	
+	{
 		// no need to store data if weight is small
 		if (Weights[DriverTriangleIndex] > UE_KINDA_SMALL_NUMBER)
 		{
@@ -49,9 +48,9 @@ FSourceVertexData::FSourceVertexData(const FVector3f& SourceVertex,
 			TriangleData.InverseDistanceWeight = Weights[DriverTriangleIndex];
 			TriangleData.Triangle = DriverTriangle;
 			TriangleData.BarycentricCoords = BarycentricCoordinates(ClosestPoint, A, B, C);
-			TriangleData.TangentIndex = GetTriangleTangentIndex(ClosestPoint, A, B, C);
-			TriangleData.InvMatrix = GetTriangleMatrix(ClosestPoint, A, B, C, TriangleData.TangentIndex).Inverse();
-			
+			TriangleData.TangentLocalIndex = GetTriangleTangentLocalIndex(ClosestPoint, A, B, C);
+			TriangleData.InvMatrix = GetTriangleMatrix(ClosestPoint, A, B, C, TriangleData.TangentLocalIndex).Inverse();
+
 			// Interpolate SkinWeights with Barycentric Coords
 			const TArray<VertexSkinWeightMax> SkinWeights = { DriverSkinWeights[TriangleData.Triangle.X], DriverSkinWeights[TriangleData.Triangle.Y], DriverSkinWeights[TriangleData.Triangle.Z] };
 			const TArray<float> BarycentricWeights = { TriangleData.BarycentricCoords.X, TriangleData.BarycentricCoords.Y, TriangleData.BarycentricCoords.Z };
@@ -60,11 +59,11 @@ FSourceVertexData::FSourceVertexData(const FVector3f& SourceVertex,
 			DriverTriangleData.Add(TriangleData);
 		}
 	}
-
 }
 
 
-FSourceMeshToDriverMesh::FSourceMeshToDriverMesh(const UStaticMesh* StaticMesh, const int32 StaticMeshLODIndex, const USkeletalMesh* SkeletalMesh, const int32 SkeletalMeshLODIndex)
+void FSourceMeshToDriverMesh::Update(const UStaticMesh* StaticMesh, const int32 StaticMeshLODIndex, 
+	const USkeletalMesh* SkeletalMesh, const int32 SkeletalMeshLODIndex)
 {
 	check(StaticMesh);
 	check(SkeletalMesh);
@@ -84,15 +83,14 @@ FSourceMeshToDriverMesh::FSourceMeshToDriverMesh(const UStaticMesh* StaticMesh,
 	// Allocate
 	SourceVerticesData.SetNumZeroed(NumSourceVertices); // note this is initializing values as zero
 	
-	// Get SourceVertex > DriverTriangle Data
-	//for (int32 SourceVertexIndex = 0; SourceVertexIndex < NumSourceVertices; SourceVertexIndex++)
+	// Get SourceVertex -> DriverTriangle Data
 	ParallelFor(NumSourceVertices, [&](int32 SourceVertexIndex)
-	{
-		const FVector3f& SourceVertex = SourceVertices[SourceVertexIndex];
-		const FVector3f& SourceNormal = SourceNormals[SourceVertexIndex];
+	{	
+		// Create Mapping from StaticMesh Vertex to SkeletalMesh Triangles
+		SourceVerticesData[SourceVertexIndex].Update(SourceVertices[SourceVertexIndex], 
+			DriverVertices, DriverTriangles, DriverSkinWeights, 1.f /*Sigma*/);
 
-		// Create Mapping from StaticMesh Vertex to SkeletalMesh
-		SourceVerticesData[SourceVertexIndex] = FSourceVertexData(SourceVertex, DriverVertices, DriverTriangles, DriverSkinWeights);
+		// UE_LOG(LogTemp, Warning, TEXT("Vertex: %i NumTriangles: %i."), SourceVertexIndex, SourceVerticesData[SourceVertexIndex].DriverTriangleData.Num());
 
 	});	// end ParallelFor
 }
@@ -142,9 +140,9 @@ void FSourceMeshToDriverMesh::DeformVerticesAndNormals(const TArray<FVector3f>&
 			const FVector3f Point = PointAtBarycentricCoordinates(A, B, C, BarycentricCoords);
 
 			// Grt Driver Triangle Matrix
-			const uint8& TangentIndex = TriangleData.TangentIndex;
+			const uint8& TangentLocalIndex = TriangleData.TangentLocalIndex;
 			const FMatrix44f& DriverTriangleInvMatrix = TriangleData.InvMatrix;
-			const FMatrix44f DriverTriangleMatrix = GetTriangleMatrix(Point, A, B, C, TangentIndex);
+			const FMatrix44f DriverTriangleMatrix = GetTriangleMatrix(Point, A, B, C, TangentLocalIndex);
 
 			// Get SourceVertex -> DriverTriangle InverseDistanceWeight
 			const float& InverseDistanceWeight = TriangleData.InverseDistanceWeight;
diff --git a/Engine/Plugins/Experimental/AnimToTexture/Source/AnimToTextureEditor/Private/AnimToTextureSkeletalMesh.cpp b/Engine/Plugins/Experimental/AnimToTexture/Source/AnimToTextureEditor/Private/AnimToTextureSkeletalMesh.cpp
--- a/Engine/Plugins/Experimental/AnimToTexture/Source/AnimToTextureEditor/Private/AnimToTextureSkeletalMesh.cpp
+++ b/Engine/Plugins/Experimental/AnimToTexture/Source/AnimToTextureEditor/Private/AnimToTextureSkeletalMesh.cpp
@@ -395,7 +395,8 @@ void GetSkinnedVertices(const USkeletalMeshComponent* SkeletalMeshComponent, con
 	};
 };
 
-FVector3f FindClosestPointToTriangle(const FVector3f& P, const FVector3f& A, const FVector3f& B, const FVector3f& C)
+FVector3f FindClosestPointToTriangle(const FVector3f& P, const FVector3f& A, const FVector3f& B, const FVector3f& C, 
+	int32& OutOnPointLocalIndex)
 {
 	const FVector3f AB = B - A;
 	const FVector3f AC = C - A;
@@ -405,6 +406,7 @@ FVector3f FindClosestPointToTriangle(const FVector3f& P, const FVector3f& A, con
 	const float D2 = FVector3f::DotProduct(AC, AP);
 	if (D1 <= 0.f && D2 <= 0.f)
 	{
+		OutOnPointLocalIndex = 0;
 		return A;
 	}
 
@@ -413,6 +415,7 @@ FVector3f FindClosestPointToTriangle(const FVector3f& P, const FVector3f& A, con
 	const float D4 = FVector3f::DotProduct(AC, BP);
 	if (D3 >= 0.f && D4 <= D3)
 	{
+		OutOnPointLocalIndex = 1;
 		return B;
 	}
 
@@ -421,12 +424,14 @@ FVector3f FindClosestPointToTriangle(const FVector3f& P, const FVector3f& A, con
 	const float D6 = FVector3f::DotProduct(AC, CP);
 	if (D6 >= 0.f && D5 <= D6)
 	{
+		OutOnPointLocalIndex = 2;
 		return C;
 	}
 
 	const float VC = D1 * D4 - D3 * D2;
 	if (VC <= 0.f && D1 >= 0.f && D3 <= 0.f)
 	{
+		OutOnPointLocalIndex = INDEX_NONE;
 		const float V = D1 / (D1 - D3);
 		return A + V * AB;
 	}
@@ -434,6 +439,7 @@ FVector3f FindClosestPointToTriangle(const FVector3f& P, const FVector3f& A, con
 	const float VB = D5 * D2 - D1 * D6;
 	if (VB <= 0.f && D2 >= 0.f && D6 <= 0.f)
 	{
+		OutOnPointLocalIndex = INDEX_NONE;
 		const float V = D2 / (D2 - D6);
 		return A + V * AC;
 	}
@@ -441,10 +447,12 @@ FVector3f FindClosestPointToTriangle(const FVector3f& P, const FVector3f& A, con
 	const float VA = D3 * D6 - D5 * D4;
 	if (VA <= 0.f && (D4 - D3) >= 0.f && (D5 - D6) >= 0.f)
 	{
+		OutOnPointLocalIndex = INDEX_NONE;
 		const float V = (D4 - D3) / ((D4 - D3) + (D5 - D6));
 		return B + V * (C - B);
 	}
 
+	OutOnPointLocalIndex = INDEX_NONE;
 	const float Denom = 1.0f / (VA + VB + VC);
 	const float V = VB * Denom;
 	const float W = VC * Denom;
@@ -459,40 +467,39 @@ FVector3f GetTriangleNormal(const FVector3f& A, const FVector3f& B, const FVecto
 	return FVector3f::CrossProduct(V0, V1); // .GetSafeNormal();
 }
 
-uint8 GetTriangleTangentIndex(const FVector3f& P, const FVector3f& A, const FVector3f& B, const FVector3f& C)
+uint8 GetTriangleTangentLocalIndex(const FVector3f& P, const FVector3f& A, const FVector3f& B, const FVector3f& C)
 {
-	TArray<FVector3f> TangentVectors = { A - P, B - P, C - P };
+	const TArray<float> Distances = { (A - P).Length(), (B - P).Length(), (C - P).Length() };
 
 	float MaxDistance = 0.f;
-	uint8 TangentIndex = 0;
-	for (uint8 Index=0; Index < 3; Index++)
+	uint8 TangentLocalIndex = 0;
+	for (uint8 Index=0; Index < Distances.Num(); Index++)
 	{
-		const float Distance = TangentVectors[Index].Length();
-		if (Distance > MaxDistance)
+		if (Distances[Index] > MaxDistance)
 		{
-			TangentIndex = Index;
-			MaxDistance = Distance;
+			TangentLocalIndex = Index;
+			MaxDistance = Distances[Index];
 		}
 	}
 
-	return TangentIndex;
+	return TangentLocalIndex;
 }
 
 
 
-FMatrix44f GetTriangleMatrix(const FVector3f& P, const FVector3f& A, const FVector3f& B, const FVector3f& C, const uint8 TangentIndex)
+FMatrix44f GetTriangleMatrix(const FVector3f& P, const FVector3f& A, const FVector3f& B, const FVector3f& C, const uint8 TangentLocalIndex)
 {
 	// Tangent Vector
 	FVector3f V0;
-	if (TangentIndex == 0)
+	if (TangentLocalIndex == 0)
 	{
 		V0 = A - P;
 	}
-	else if (TangentIndex == 1)
+	else if (TangentLocalIndex == 1)
 	{
 		V0 = B - P;
 	}
-	else if (TangentIndex == 2)
+	else if (TangentLocalIndex == 2)
 	{
 		V0 = C - P;
 	}
@@ -500,10 +507,11 @@ FMatrix44f GetTriangleMatrix(const FVector3f& P, const FVector3f& A, const FVect
 	const FVector3f V1 = GetTriangleNormal(A, B, C);
 	const FVector3f V2 = FVector3f::CrossProduct(V0, V1);
 
-	FMatrix44f Matrix = FMatrix44f::Identity; // SetAxes doesn't set [3] elements.
-	Matrix.SetAxes(&V0, &V1, &V2, &P);
+	//FMatrix44f Matrix = FMatrix44f::Identity; // SetAxes doesn't set [3] elements.
+	//Matrix.SetAxes(&V0, &V1, &V2, &P);
+	//return Matrix;
 
-	return Matrix;
+	return FMatrix44f(V0, V1, V2, P);
 }
 
 FVector3f BarycentricCoordinates(const FVector3f& P, const FVector3f& A, const FVector3f& B, const FVector3f& C)
@@ -525,7 +533,7 @@ FVector3f BarycentricCoordinates(const FVector3f& P, const FVector3f& A, const F
 	return FVector3f(U, V, W);
 }
 
-void InverseDistanceWeights(const FVector3f& Point, const TArray<FVector3f>& Points,
+int32 InverseDistanceWeights(const FVector3f& Point, const TArray<FVector3f>& Points,
 	TArray<float>& OutWeights, float Sigma)
 {
 	// Allocate
@@ -544,7 +552,7 @@ void InverseDistanceWeights(const FVector3f& Point, const TArray<FVector3f>& Poi
 		if (Distance < UE_KINDA_SMALL_NUMBER)
 		{
 			OutWeights[Index] = 1.f;
-			return;
+			return Index;
 		}
 
 		InverseDistances[Index] = 1.f / FMath::Pow(Distance, Sigma);
@@ -560,6 +568,8 @@ void InverseDistanceWeights(const FVector3f& Point, const TArray<FVector3f>& Poi
 	{
 		OutWeights[Index] = InverseDistances[Index] / SumInverseDistance;
 	}
+
+	return INDEX_NONE;
 }
 
 
diff --git a/Engine/Plugins/Experimental/AnimToTexture/Source/AnimToTextureEditor/Public/AnimToTextureBPLibrary.h b/Engine/Plugins/Experimental/AnimToTexture/Source/AnimToTextureEditor/Public/AnimToTextureBPLibrary.h
--- a/Engine/Plugins/Experimental/AnimToTexture/Source/AnimToTextureEditor/Public/AnimToTextureBPLibrary.h
+++ b/Engine/Plugins/Experimental/AnimToTexture/Source/AnimToTextureEditor/Public/AnimToTextureBPLibrary.h
@@ -57,6 +57,15 @@ class UAnimToTextureBPLibrary : public UBlueprintFunctionLibrary
 
 private:
 
+	// Runs some validations for the assets in DataAsset
+	// Returns false if there is any problems with the data, warnings will be printed in Log
+	static bool CheckDataAsset(const UAnimToTextureDataAsset* DataAsset,
+		int32& OutSocketIndex, TArray<FAnimToTextureAnimSequenceInfo>& OutAnimSequences);
+
+	// Returns Start, EndFrame and NumFrames in Animation
+	static int32 GetAnimationFrameRange(const FAnimToTextureAnimSequenceInfo& Animation, 
+		int32& OutStartFrame, int32& OutEndFrame);
+
 	// Get Vertex and Normals from Current Pose
 	// The VertexDelta is returned from the RefPose
 	static void GetVertexDeltasAndNormals(const USkeletalMeshComponent* SkeletalMeshComponent, const int32 LODIndex, 
diff --git a/Engine/Plugins/Experimental/AnimToTexture/Source/AnimToTextureEditor/Public/AnimToTextureMeshMapping.h b/Engine/Plugins/Experimental/AnimToTexture/Source/AnimToTextureEditor/Public/AnimToTextureMeshMapping.h
--- a/Engine/Plugins/Experimental/AnimToTexture/Source/AnimToTextureEditor/Public/AnimToTextureMeshMapping.h
+++ b/Engine/Plugins/Experimental/AnimToTexture/Source/AnimToTextureEditor/Public/AnimToTextureMeshMapping.h
@@ -10,7 +10,7 @@ namespace AnimToTexture_Private
 
 struct FSourceVertexDriverTriangleData
 {
-	uint8               TangentIndex;
+	uint8               TangentLocalIndex;
 	float               InverseDistanceWeight;
 	FIntVector3         Triangle;
 	FVector3f           ClosestPoint;
@@ -24,7 +24,8 @@ class FSourceVertexData
 public:
 
 	FSourceVertexData() = default;	
-	FSourceVertexData(const FVector3f& SourceVertex,
+	
+	void Update(const FVector3f& SourceVertex,
 		const TArray<FVector3f>& DriverVertices, const TArray<FIntVector3>& DriverTriangles, const TArray<VertexSkinWeightMax>& DriverSkinWeights, 
 		const float Sigma=1.f);
 
@@ -39,7 +40,8 @@ class FSourceMeshToDriverMesh
 public:
 
 	FSourceMeshToDriverMesh() = default;
-	FSourceMeshToDriverMesh(const UStaticMesh* StaticMesh, const int32 StaticMeshLODIndex,
+	
+	void Update(const UStaticMesh* StaticMesh, const int32 StaticMeshLODIndex,
 		const USkeletalMesh* SkeletalMesh, const int32 SkeletalMeshLODIndex);
 
 	// Returns Number of Source Vertices
diff --git a/Engine/Plugins/Experimental/AnimToTexture/Source/AnimToTextureEditor/Public/AnimToTextureSkeletalMesh.h b/Engine/Plugins/Experimental/AnimToTexture/Source/AnimToTextureEditor/Public/AnimToTextureSkeletalMesh.h
--- a/Engine/Plugins/Experimental/AnimToTexture/Source/AnimToTextureEditor/Public/AnimToTextureSkeletalMesh.h
+++ b/Engine/Plugins/Experimental/AnimToTexture/Source/AnimToTextureEditor/Public/AnimToTextureSkeletalMesh.h
@@ -62,8 +62,9 @@ bool HasBone(const USkeletalMesh* SkeletalMesh, const FName& Bone);
    Only the RawBones are returned (no virtual bones) */
 void GetBoneNames(const USkeletalMesh* SkeletalMesh, TArray<FName>& OutNames);
 
-/* Computes closest point to triangle */
-FVector3f FindClosestPointToTriangle(const FVector3f& Point, const FVector3f& PointA, const FVector3f& PointB, const FVector3f& PointC);
+/* Computes closest point to triangle. Returns index of point [0-2] if right on top, else INDEX_NONE */
+FVector3f FindClosestPointToTriangle(const FVector3f& Point, const FVector3f& PointA, const FVector3f& PointB, const FVector3f& PointC, 
+	int32& OutOnPointLocalIndex);
 
 /* Computes Barycentric coordinates from point to triangle */
 FVector3f BarycentricCoordinates(const FVector3f& Point, const FVector3f& PointA, const FVector3f& PointB, const FVector3f& PointC);
@@ -74,18 +75,19 @@ FVector3f PointAtBarycentricCoordinates(const FVector3f& PointA, const FVector3f
 /* Computes Triangle Normal (not normalized) */
 FVector3f GetTriangleNormal(const FVector3f& PointA, const FVector3f& PointB, const FVector3f& PointC);
 
-/* Returns TangentIndex for Triangle 
-* TangentIndex 0: PointA - PointP
-* TangentIndex 1: PointB - PointP
-* TangentIndex 2: PointC - PointP */
-uint8 GetTriangleTangentIndex(const FVector3f& Point, const FVector3f& PointA, const FVector3f& PointB, const FVector3f& PointC);
+/* Returns TangentLocalIndex for Triangle 
+* TangentLocalIndex 0: PointA - PointP
+* TangentLocalIndex 1: PointB - PointP
+* TangentLocalIndex 2: PointC - PointP */
+uint8 GetTriangleTangentLocalIndex(const FVector3f& Point, const FVector3f& PointA, const FVector3f& PointB, const FVector3f& PointC);
 
 /* Computes Triangle Matrix */
-FMatrix44f GetTriangleMatrix(const FVector3f& Point, const FVector3f& PointA, const FVector3f& PointB, const FVector3f& PointC, const uint8 TangentIndex);
+FMatrix44f GetTriangleMatrix(const FVector3f& Point, const FVector3f& PointA, const FVector3f& PointB, const FVector3f& PointC, const uint8 TangentLocalIndex);
 
 // Returns a list of weights from Point to each of the given Points.
 // The returned Weights are the size of Points and normalized (sum to 1.0) 
-void InverseDistanceWeights(const FVector3f& Point, const TArray<FVector3f>& Points,
+// If one of the points is right on top (no need to interpolate), it will return index, else INDEX_NONE
+int32 InverseDistanceWeights(const FVector3f& Point, const TArray<FVector3f>& Points,
 	TArray<float>& OutWeights, float Sigma = 1.f);
 
 } // end namespace AnimToTexture_Private
