[AnimToTexture] Fixing LOCTEXT_NAMESPACE conflict

diff --git a/Engine/Plugins/Experimental/AnimToTexture/Source/AnimToTextureEditor/Private/AnimToTextureBPLibrary.cpp b/Engine/Plugins/Experimental/AnimToTexture/Source/AnimToTextureEditor/Private/AnimToTextureBPLibrary.cpp
--- a/Engine/Plugins/Experimental/AnimToTexture/Source/AnimToTextureEditor/Private/AnimToTextureBPLibrary.cpp
+++ b/Engine/Plugins/Experimental/AnimToTexture/Source/AnimToTextureEditor/Private/AnimToTextureBPLibrary.cpp
@@ -20,7 +20,7 @@
 #include "MaterialEditingLibrary.h"
 #include "Misc/ScopedSlowTask.h"
 
-#define LOCTEXT_NAMESPACE "UAnimToTextureBPLibrary"
+#define LOCTEXT_NAMESPACE "AnimToTextureEditor"
 
 using namespace AnimToTexture_Private;
 
@@ -54,7 +54,7 @@ bool UAnimToTextureBPLibrary::AnimationToTexture(UAnimToTextureDataAsset* DataAs
 	//
 	FSourceMeshToDriverMesh Mapping;
 	{
-		FScopedSlowTask ProgressBar(1.f, LOCTEXT("ProcessingAnimSequence", "Processing StaticMesh -> SkeletalMesh Mapping ..."), true /*Enabled*/);
+		FScopedSlowTask ProgressBar(1.f, LOCTEXT("ProcessingMapping", "Processing StaticMesh -> SkeletalMesh Mapping ..."), true /*Enabled*/);
 		ProgressBar.MakeDialog(false /*bShowCancelButton*/, false /*bAllowInPIE*/);
 
 		Mapping.Update(DataAsset->GetStaticMesh(), DataAsset->StaticLODIndex,
diff --git a/Engine/Plugins/Experimental/AnimToTexture/Source/AnimToTextureEditor/Private/AnimToTextureEditorModule.cpp b/Engine/Plugins/Experimental/AnimToTexture/Source/AnimToTextureEditor/Private/AnimToTextureEditorModule.cpp
--- a/Engine/Plugins/Experimental/AnimToTexture/Source/AnimToTextureEditor/Private/AnimToTextureEditorModule.cpp
+++ b/Engine/Plugins/Experimental/AnimToTexture/Source/AnimToTextureEditor/Private/AnimToTextureEditorModule.cpp
@@ -6,7 +6,7 @@
 #include "AssetToolsModule.h"
 #include "MessageLogModule.h"
 
-#define LOCTEXT_NAMESPACE "FAnimToTextureEditorModule"
+#define LOCTEXT_NAMESPACE "AnimToTextureEditor"
 
 DEFINE_LOG_CATEGORY(LogAnimToTextureEditor);
 
@@ -35,4 +35,3 @@ void FAnimToTextureEditorModule::ShutdownModule()
 IMPLEMENT_MODULE(FAnimToTextureEditorModule, AnimToTextureEditor)
 
 #undef LOCTEXT_NAMESPACE
-	
\ No newline at end of file
