[AnimToTexture] Adding MaterialInstance static params in DataAsset

diff --git a/Engine/Plugins/Experimental/AnimToTexture/Source/AnimToTexture/Public/AnimToTextureDataAsset.h b/Engine/Plugins/Experimental/AnimToTexture/Source/AnimToTexture/Public/AnimToTextureDataAsset.h
--- a/Engine/Plugins/Experimental/AnimToTexture/Source/AnimToTexture/Public/AnimToTextureDataAsset.h
+++ b/Engine/Plugins/Experimental/AnimToTexture/Source/AnimToTexture/Public/AnimToTextureDataAsset.h
@@ -178,7 +178,7 @@ public:
 	        This is the preferred method if meshes share skeleton.
 	*/
 	UPROPERTY(EditAnywhere, BlueprintReadWrite, Category = "Texture")
-	EAnimToTextureMode Mode;
+	EAnimToTextureMode Mode = EAnimToTextureMode::Bone;
 
 	/**
 	* Texture for storing vertex positions
@@ -238,6 +238,39 @@ public:
 	UPROPERTY(EditAnywhere, BlueprintReadWrite, Category = "Animation")
 	TArray<FAnimToTextureAnimSequenceInfo> AnimSequences;
 
+	// ------------------------------------------------------
+	// Material Layer Static Switches
+
+	/**
+	* AutoPlay will use Engine Time for driving the animation.
+	* This will be used by UpdateMaterialInstanceFromDataAsset and AssetActions for setting MaterialInstance static switches
+	*/
+	UPROPERTY(EditAnywhere, BlueprintReadWrite, Category = "Material")
+	bool bAutoPlay = true;
+
+	/**
+	* AnimationIndex Index of the animation to play.
+	* This will internally set Start and End Frame when using AutoPlay.
+	* This will be used by UpdateMaterialInstanceFromDataAsset and AssetActions for setting MaterialInstance static switches
+	*/
+	UPROPERTY(EditAnywhere, BlueprintReadWrite, Category = "Material", meta = (EditCondition = "bAutoPlay"))
+	int32 AnimationIndex = 0;
+
+	/**
+	* Frame to play
+	* When not using AutoPlay, user is responsible of setting the frame.
+	* This will be used by UpdateMaterialInstanceFromDataAsset and AssetActions for setting MaterialInstance static switches
+	*/
+	UPROPERTY(EditAnywhere, BlueprintReadWrite, Category = "Material", meta = (EditCondition = "!bAutoPlay"))
+	int32 Frame = 0;
+
+	/**
+	* Number of Bone Influences for deformation. More influences will produce smoother results at the cost of performance.
+	* This will be used by UpdateMaterialInstanceFromDataAsset and AssetActions for setting MaterialInstance static switches
+	*/
+	UPROPERTY(EditAnywhere, BlueprintReadWrite, Category = "Material", meta = (EditCondition = "Mode == EAnimToTextureMode::Bone", EditConditionHides))
+	EAnimToTextureNumBoneInfluences NumBoneInfluences = EAnimToTextureNumBoneInfluences::Four;
+
 	// ------------------------------------------------------
 	// Info
 
diff --git a/Engine/Plugins/Experimental/AnimToTexture/Source/AnimToTextureEditor/Private/AnimToTextureAssetActions.cpp b/Engine/Plugins/Experimental/AnimToTexture/Source/AnimToTextureEditor/Private/AnimToTextureAssetActions.cpp
--- a/Engine/Plugins/Experimental/AnimToTexture/Source/AnimToTextureEditor/Private/AnimToTextureAssetActions.cpp
+++ b/Engine/Plugins/Experimental/AnimToTexture/Source/AnimToTextureEditor/Private/AnimToTextureAssetActions.cpp
@@ -87,11 +87,7 @@ void FAnimToTextureAssetActions::RunAnimToTexture(TArray<TWeakObjectPtr<UAnimToT
 				{
 					if (UMaterialInstanceConstant* MaterialInstanceConstant = Cast<UMaterialInstanceConstant>(StaticMaterial.MaterialInterface))
 					{
-						UAnimToTextureBPLibrary::UpdateMaterialInstanceFromDataAsset(DataAsset, MaterialInstanceConstant,
-							/*bAutoPlay*/ true, 
-							/*AnimationIndex*/ 0, 
-							EAnimToTextureNumBoneInfluences::Four,
-							EMaterialParameterAssociation::LayerParameter);
+						UAnimToTextureBPLibrary::UpdateMaterialInstanceFromDataAsset(DataAsset, MaterialInstanceConstant, EMaterialParameterAssociation::LayerParameter);
 					}
 				}
 			}
diff --git a/Engine/Plugins/Experimental/AnimToTexture/Source/AnimToTextureEditor/Private/AnimToTextureBPLibrary.cpp b/Engine/Plugins/Experimental/AnimToTexture/Source/AnimToTextureEditor/Private/AnimToTextureBPLibrary.cpp
--- a/Engine/Plugins/Experimental/AnimToTexture/Source/AnimToTextureEditor/Private/AnimToTextureBPLibrary.cpp
+++ b/Engine/Plugins/Experimental/AnimToTexture/Source/AnimToTextureEditor/Private/AnimToTextureBPLibrary.cpp
@@ -600,9 +600,8 @@ int32 UAnimToTextureBPLibrary::GetBonePositionsAndRotations(const USkeletalMeshC
 }
 
 
-void UAnimToTextureBPLibrary::UpdateMaterialInstanceFromDataAsset(UAnimToTextureDataAsset* DataAsset, UMaterialInstanceConstant* MaterialInstance, 
-	const bool bAutoPlay, const int32 AnimationIndex, 
-	const EAnimToTextureNumBoneInfluences NumBoneInfluences, const EMaterialParameterAssociation MaterialParameterAssociation)
+void UAnimToTextureBPLibrary::UpdateMaterialInstanceFromDataAsset(const UAnimToTextureDataAsset* DataAsset, UMaterialInstanceConstant* MaterialInstance, 
+	const EMaterialParameterAssociation MaterialParameterAssociation)
 {
 	check(DataAsset);
 	check(MaterialInstance);
@@ -675,7 +674,7 @@ void UAnimToTextureBPLibrary::UpdateMaterialInstanceFromDataAsset(UAnimToTexture
 		UMaterialEditingLibrary::SetMaterialInstanceTextureParameterValue(MaterialInstance, AnimToTextureParamNames::BoneWeightsTexture, DataAsset->GetBoneWeightTexture(), MaterialParameterAssociation);
 
 		// Num Influences
-		switch (NumBoneInfluences)
+		switch (DataAsset->NumBoneInfluences)
 		{
 			case EAnimToTextureNumBoneInfluences::One:
 				UMaterialEditingLibrary::SetMaterialInstanceStaticSwitchParameterValue(MaterialInstance, AnimToTextureParamNames::UseTwoInfluences, false, MaterialParameterAssociation);
@@ -690,15 +689,32 @@ void UAnimToTextureBPLibrary::UpdateMaterialInstanceFromDataAsset(UAnimToTexture
 				UMaterialEditingLibrary::SetMaterialInstanceStaticSwitchParameterValue(MaterialInstance, AnimToTextureParamNames::UseFourInfluences, true, MaterialParameterAssociation);
 				break;
 		}
-
 	}
 
 	// AutoPlay
-	if (bAutoPlay && DataAsset->Animations.IsValidIndex(AnimationIndex))
+	UMaterialEditingLibrary::SetMaterialInstanceStaticSwitchParameterValue(MaterialInstance, AnimToTextureParamNames::AutoPlay, DataAsset->bAutoPlay, MaterialParameterAssociation);
+	if (DataAsset->bAutoPlay)
+	{
+		if (DataAsset->Animations.IsValidIndex(DataAsset->AnimationIndex))
+		{
+			UMaterialEditingLibrary::SetMaterialInstanceScalarParameterValue(MaterialInstance, AnimToTextureParamNames::StartFrame, DataAsset->Animations[DataAsset->AnimationIndex].StartFrame, MaterialParameterAssociation);
+			UMaterialEditingLibrary::SetMaterialInstanceScalarParameterValue(MaterialInstance, AnimToTextureParamNames::EndFrame, DataAsset->Animations[DataAsset->AnimationIndex].EndFrame, MaterialParameterAssociation);
+		}
+		else
+		{
+			UE_LOG(LogAnimToTextureEditor, Warning, TEXT("Invalid AnimationIndex: %i"), DataAsset->AnimationIndex);
+		}
+	}
+	else
 	{
-		UMaterialEditingLibrary::SetMaterialInstanceStaticSwitchParameterValue(MaterialInstance, AnimToTextureParamNames::AutoPlay, true, MaterialParameterAssociation);
-		UMaterialEditingLibrary::SetMaterialInstanceScalarParameterValue(MaterialInstance, AnimToTextureParamNames::StartFrame, DataAsset->Animations[AnimationIndex].StartFrame, MaterialParameterAssociation);
-		UMaterialEditingLibrary::SetMaterialInstanceScalarParameterValue(MaterialInstance, AnimToTextureParamNames::EndFrame, DataAsset->Animations[AnimationIndex].EndFrame, MaterialParameterAssociation);
+		if (DataAsset->Frame >= 0 && DataAsset->Frame < DataAsset->NumFrames)
+		{
+			UMaterialEditingLibrary::SetMaterialInstanceScalarParameterValue(MaterialInstance, AnimToTextureParamNames::Frame, DataAsset->Frame, MaterialParameterAssociation);
+		}
+		else
+		{
+			UE_LOG(LogAnimToTextureEditor, Warning, TEXT("Frame out of range: %i"), DataAsset->Frame);
+		}
 	}
 	
 	// NumFrames
diff --git a/Engine/Plugins/Experimental/AnimToTexture/Source/AnimToTextureEditor/Public/AnimToTextureBPLibrary.h b/Engine/Plugins/Experimental/AnimToTexture/Source/AnimToTextureEditor/Public/AnimToTextureBPLibrary.h
--- a/Engine/Plugins/Experimental/AnimToTexture/Source/AnimToTextureEditor/Public/AnimToTextureBPLibrary.h
+++ b/Engine/Plugins/Experimental/AnimToTexture/Source/AnimToTextureEditor/Public/AnimToTextureBPLibrary.h
@@ -44,13 +44,9 @@ class UAnimToTextureBPLibrary : public UBlueprintFunctionLibrary
 
 	/**
 	* Updates a material's parameters to match those of an AnimToTexture DataAsset
-	* @param AnimationIndex Index of the animation to play.
-	* @param NumBoneInfluences Number of skinning influences. More influences will produce smoother deformation but slower performance. This is only used in Bone Mode.
 	*/
 	UFUNCTION(BlueprintCallable, meta = (Category = "AnimToTexture"))
-	static void UpdateMaterialInstanceFromDataAsset(UAnimToTextureDataAsset* DataAsset, class UMaterialInstanceConstant* MaterialInstance, 
-		const bool bAutoPlay=true, const int32 AnimationIndex=0, 
-		const EAnimToTextureNumBoneInfluences NumBoneInfluences = EAnimToTextureNumBoneInfluences::One,
+	static void UpdateMaterialInstanceFromDataAsset(const UAnimToTextureDataAsset* DataAsset, class UMaterialInstanceConstant* MaterialInstance,
 		const EMaterialParameterAssociation MaterialParameterAssociation = EMaterialParameterAssociation::LayerParameter);
 
 #endif // WITH_EDITOR
