[AnimToTexture] Only update params if main function is successful

diff --git a/Engine/Plugins/Experimental/AnimToTexture/Source/AnimToTextureEditor/Private/AnimToTextureAssetActions.cpp b/Engine/Plugins/Experimental/AnimToTexture/Source/AnimToTextureEditor/Private/AnimToTextureAssetActions.cpp
--- a/Engine/Plugins/Experimental/AnimToTexture/Source/AnimToTextureEditor/Private/AnimToTextureAssetActions.cpp
+++ b/Engine/Plugins/Experimental/AnimToTexture/Source/AnimToTextureEditor/Private/AnimToTextureAssetActions.cpp
@@ -78,16 +78,17 @@ void FAnimToTextureAssetActions::RunAnimToTexture(TArray<TWeakObjectPtr<UAnimToT
 		if (UAnimToTextureDataAsset* DataAsset = (*ObjIt).Get())
 		{
 			// Create UVs and Textures
-			UAnimToTextureBPLibrary::AnimationToTexture(DataAsset);
-
-			// Update Material Instances (if Possible)
-			if (UStaticMesh* StaticMesh = DataAsset->GetStaticMesh())
+			if (UAnimToTextureBPLibrary::AnimationToTexture(DataAsset))
 			{
-				for (FStaticMaterial& StaticMaterial : StaticMesh->GetStaticMaterials())
+				// Update Material Instances (if Possible)
+				if (UStaticMesh* StaticMesh = DataAsset->GetStaticMesh())
 				{
-					if (UMaterialInstanceConstant* MaterialInstanceConstant = Cast<UMaterialInstanceConstant>(StaticMaterial.MaterialInterface))
+					for (FStaticMaterial& StaticMaterial: StaticMesh->GetStaticMaterials())
 					{
-						UAnimToTextureBPLibrary::UpdateMaterialInstanceFromDataAsset(DataAsset, MaterialInstanceConstant, EMaterialParameterAssociation::LayerParameter);
+						if (UMaterialInstanceConstant* MaterialInstanceConstant = Cast<UMaterialInstanceConstant>(StaticMaterial.MaterialInterface))
+						{
+							UAnimToTextureBPLibrary::UpdateMaterialInstanceFromDataAsset(DataAsset, MaterialInstanceConstant, EMaterialParameterAssociation::LayerParameter);
+						}
 					}
 				}
 			}
