[AnimToTexture] Fixed InterpolateVertexSkinWeights bug with large influence indices

diff --git a/Engine/Plugins/Experimental/AnimToTexture/Source/AnimToTextureEditor/Private/AnimToTextureSkeletalMesh.cpp b/Engine/Plugins/Experimental/AnimToTexture/Source/AnimToTextureEditor/Private/AnimToTextureSkeletalMesh.cpp
--- a/Engine/Plugins/Experimental/AnimToTexture/Source/AnimToTextureEditor/Private/AnimToTextureSkeletalMesh.cpp
+++ b/Engine/Plugins/Experimental/AnimToTexture/Source/AnimToTextureEditor/Private/AnimToTextureSkeletalMesh.cpp
@@ -267,12 +267,12 @@ void InterpolateVertexSkinWeights(const TArray<VertexSkinWeightMax>& SkinWeights
 	check(SkinWeights.Num() == Weights.Num())
 
 	// Reset Values
-	OutVertexSkinWeights.BoneWeights = TStaticArray<uint8, MAX_TOTAL_INFLUENCES>(InPlace, 0.f);
+	OutVertexSkinWeights.BoneWeights = TStaticArray<uint8, MAX_TOTAL_INFLUENCES>(InPlace, 0);
 	OutVertexSkinWeights.MeshBoneIndices = TStaticArray<uint16, MAX_TOTAL_INFLUENCES>(InPlace, 0);
 
 	// Create Sparse Weighted-SkinWeights
-	TMap<uint16, float> SparseSkinWeights;
-	SparseSkinWeights.Reserve(TNumericLimits<uint16>::Max());
+	TMap<uint16, float> WeightedSkinWeights;
+	WeightedSkinWeights.Reserve(TNumericLimits<uint16>::Max());
 
 	for (int32 Index = 0; Index < SkinWeights.Num(); Index++)
 	{
@@ -283,14 +283,14 @@ void InterpolateVertexSkinWeights(const TArray<VertexSkinWeightMax>& SkinWeights
 
 			const float WeightedVertexSkinWeight = (float)BoneWeight / 255.f * Weights[Index];
 
-			if (WeightedVertexSkinWeight > SMALL_NUMBER)
+			if (WeightedVertexSkinWeight > UE_KINDA_SMALL_NUMBER)
 			{
-				float* Value = SparseSkinWeights.Find(MeshBoneIndex);
+				float* Value = WeightedSkinWeights.Find(MeshBoneIndex);
 
 				// Inititialize
 				if (Value == nullptr)
 				{
-					SparseSkinWeights.Add(MeshBoneIndex, WeightedVertexSkinWeight);
+					WeightedSkinWeights.Add(MeshBoneIndex, WeightedVertexSkinWeight);
 				}
 				// Accumulate Weighted VertexSkinWeight
 				else
@@ -303,20 +303,19 @@ void InterpolateVertexSkinWeights(const TArray<VertexSkinWeightMax>& SkinWeights
 
 	// Convert Weights to uint8 (and get them ready for Sort)
 	TArray<TPair<uint8, uint16>> SortedVertexSkinWeights;
-	SortedVertexSkinWeights.SetNumZeroed(MAX_TOTAL_INFLUENCES); // note allocate the max 12 influences
-	int32 Index = 0;
-	for (const auto& Item : SparseSkinWeights)
+	SortedVertexSkinWeights.Reserve(TNumericLimits<uint16>::Max());
+	
+	for (const auto& Item : WeightedSkinWeights)
 	{
 		const uint8 BoneWeight = (uint8)FMath::RoundToInt(Item.Value * 255.f);
 		const uint16& MeshBoneIndex = Item.Key;
-		SortedVertexSkinWeights[Index] = TPair<uint8, uint16>(BoneWeight, MeshBoneIndex);
-		Index++;
+		SortedVertexSkinWeights.Add(TPair<uint8, uint16>(BoneWeight, MeshBoneIndex));
 	}
 
 	// Sort Weights InPlace and reverse
 	SortedVertexSkinWeights.Sort();
 	Algo::Reverse(SortedVertexSkinWeights);
-	for (int32 InfluIndex = 0; InfluIndex < MAX_TOTAL_INFLUENCES; InfluIndex++)
+	for (int32 InfluIndex = 0; InfluIndex < FMath::Min(SortedVertexSkinWeights.Num(), MAX_TOTAL_INFLUENCES); InfluIndex++)
 	{
 		OutVertexSkinWeights.BoneWeights[InfluIndex] = SortedVertexSkinWeights[InfluIndex].Key;
 		OutVertexSkinWeights.MeshBoneIndices[InfluIndex] = SortedVertexSkinWeights[InfluIndex].Value;
