[AnimToTexture] Fixes slow mapping between StaticMesh and SkeletalMesh.

diff --git a/Engine/Plugins/Experimental/AnimToTexture/Source/AnimToTextureEditor/Private/AnimToTextureMeshMapping.cpp b/Engine/Plugins/Experimental/AnimToTexture/Source/AnimToTextureEditor/Private/AnimToTextureMeshMapping.cpp
--- a/Engine/Plugins/Experimental/AnimToTexture/Source/AnimToTextureEditor/Private/AnimToTextureMeshMapping.cpp
+++ b/Engine/Plugins/Experimental/AnimToTexture/Source/AnimToTextureEditor/Private/AnimToTextureMeshMapping.cpp
@@ -9,9 +9,7 @@ void FSourceVertexData::Update(const FVector3f& SourceVertex,
 	const TArray<FVector3f>& DriverVertices, const TArray<FIntVector3>& DriverTriangles, const TArray<VertexSkinWeightMax>& DriverSkinWeights, 
 	const float Sigma)
 {	
-	// Reserve TriangleData
 	const int32 NumDriverTriangles = DriverTriangles.Num();
-	DriverTriangleData.Reserve(NumDriverTriangles);
 
 	// Get ClosestPoint to Triangles
 	TArray<FVector3f> ClosestPoints;
@@ -32,7 +30,7 @@ void FSourceVertexData::Update(const FVector3f& SourceVertex,
 	// Get Inverse Distance Weighting from Vertex to each DriverTriangle (ClosestPoint)
 	TArray<float> Weights;
 	AnimToTexture_Private::InverseDistanceWeights(SourceVertex, ClosestPoints, Weights, Sigma);
-	
+
 	for (int32 DriverTriangleIndex = 0; DriverTriangleIndex < NumDriverTriangles; DriverTriangleIndex++)
 	{
 		// no need to store data if weight is small
