[AnimToTexture]
- Added 16bits precision for Bone Weights Texture. This removes the 256 bone limitation
- Added check for no valid animations passed
- Using FVector4f

diff --git a/Engine/Plugins/Experimental/AnimToTexture/Source/AnimToTexture/Private/AnimToTextureDataAsset.cpp b/Engine/Plugins/Experimental/AnimToTexture/Source/AnimToTexture/Private/AnimToTextureDataAsset.cpp
--- a/Engine/Plugins/Experimental/AnimToTexture/Source/AnimToTexture/Private/AnimToTextureDataAsset.cpp
+++ b/Engine/Plugins/Experimental/AnimToTexture/Source/AnimToTexture/Private/AnimToTextureDataAsset.cpp
@@ -27,12 +27,13 @@ void UAnimToTextureDataAsset::ResetInfo()
 
 	// Vertex Info
 	VertexRowsPerFrame = 1;
-	VertexMinBBox = FVector::ZeroVector;
-	VertexSizeBBox = FVector::ZeroVector;
+	VertexMinBBox = FVector3f::ZeroVector;
+	VertexSizeBBox = FVector3f::ZeroVector;
 
 	// Bone Info
+	NumBones = 0;
 	BoneRowsPerFrame = 1;
 	BoneWeightRowsPerFrame = 1;
-	BoneMinBBox = FVector::ZeroVector;
-	BoneSizeBBox = FVector::ZeroVector;
+	BoneMinBBox = FVector3f::ZeroVector;
+	BoneSizeBBox = FVector3f::ZeroVector;
 };
\ No newline at end of file
diff --git a/Engine/Plugins/Experimental/AnimToTexture/Source/AnimToTexture/Public/AnimToTextureDataAsset.h b/Engine/Plugins/Experimental/AnimToTexture/Source/AnimToTexture/Public/AnimToTextureDataAsset.h
--- a/Engine/Plugins/Experimental/AnimToTexture/Source/AnimToTexture/Public/AnimToTextureDataAsset.h
+++ b/Engine/Plugins/Experimental/AnimToTexture/Source/AnimToTexture/Public/AnimToTextureDataAsset.h
@@ -21,8 +21,9 @@ namespace AnimToTextureParamNames
 	static const FName EndFrame = TEXT("EndFrame");
 	static const FName SampleRate = TEXT("SampleRate");
 	static const FName NumFrames = TEXT("NumFrames");
-	static const FName BoundingBoxMin = TEXT("MinBBox");
-	static const FName BoundingBoxScale = TEXT("SizeBBox");
+	static const FName MinBBox = TEXT("MinBBox");
+	static const FName SizeBBox = TEXT("SizeBBox");
+	static const FName NumBones = TEXT("NumBones");
 	static const FName RowsPerFrame = TEXT("RowsPerFrame");
 	static const FName BoneWeightRowsPerFrame = TEXT("BoneWeightsRowsPerFrame");
 	static const FName VertexPositionTexture = TEXT("PositionTexture");
@@ -278,14 +279,18 @@ public:
 	UPROPERTY(VisibleAnywhere, BlueprintReadOnly, Category = "Info")
 	int32 NumFrames = 0;
 
+	/* Total Number of Bones */
+	UPROPERTY(VisibleAnywhere, BlueprintReadOnly, Category = "Info", Meta = (EditCondition = "Mode == EAnimToTextureMode::Bone", EditConditionHides))
+	int32 NumBones = 0;
+
 	UPROPERTY(VisibleAnywhere, BlueprintReadOnly, Category = "Info", Meta = (EditCondition = "Mode == EAnimToTextureMode::Vertex", EditConditionHides))
 	int32 VertexRowsPerFrame = 1;
 
 	UPROPERTY(VisibleAnywhere, BlueprintReadOnly, Category = "Info", Meta = (DisplayName = "MinBBox", EditCondition = "Mode == EAnimToTextureMode::Vertex", EditConditionHides))
-	FVector VertexMinBBox;
+	FVector3f VertexMinBBox;
 
 	UPROPERTY(VisibleAnywhere, BlueprintReadOnly, Category = "Info", Meta = (DisplayName = "SizeBBox", EditCondition = "Mode == EAnimToTextureMode::Vertex", EditConditionHides))
-	FVector VertexSizeBBox;
+	FVector3f VertexSizeBBox;
 
 	UPROPERTY(VisibleAnywhere, BlueprintReadOnly, Category = "Info", Meta = (EditCondition = "Mode == EAnimToTextureMode::Bone", EditConditionHides))
 	int32 BoneWeightRowsPerFrame = 1;
@@ -294,10 +299,10 @@ public:
 	int32 BoneRowsPerFrame = 1;
 
 	UPROPERTY(VisibleAnywhere, BlueprintReadOnly, Category = "Info", Meta = (DisplayName = "MinBBox", EditCondition = "Mode == EAnimToTextureMode::Bone", EditConditionHides))
-	FVector BoneMinBBox;
+	FVector3f BoneMinBBox;
 
 	UPROPERTY(VisibleAnywhere, BlueprintReadOnly, Category = "Info", Meta = (DisplayName = "SizeBBox", EditCondition = "Mode == EAnimToTextureMode::Bone", EditConditionHides))
-	FVector BoneSizeBBox;
+	FVector3f BoneSizeBBox;
 
 	UPROPERTY(VisibleAnywhere, BlueprintReadOnly, Category = "Info")
 	TArray<FAnimToTextureAnimInfo> Animations;
diff --git a/Engine/Plugins/Experimental/AnimToTexture/Source/AnimToTextureEditor/Private/AnimToTextureBPLibrary.cpp b/Engine/Plugins/Experimental/AnimToTexture/Source/AnimToTextureEditor/Private/AnimToTextureBPLibrary.cpp
--- a/Engine/Plugins/Experimental/AnimToTexture/Source/AnimToTextureEditor/Private/AnimToTextureBPLibrary.cpp
+++ b/Engine/Plugins/Experimental/AnimToTexture/Source/AnimToTextureEditor/Private/AnimToTextureBPLibrary.cpp
@@ -71,16 +71,15 @@ bool UAnimToTextureBPLibrary::AnimationToTexture(UAnimToTextureDataAsset* DataAs
 	// ---------------------------------------------------------------------------
 	// Get Reference Skeleton Transforms
 	//
-	int32 NumBones = INDEX_NONE;
 	TArray<FVector3f> BoneRefPositions;
-	TArray<FVector4>  BoneRefRotations;
+	TArray<FVector4f> BoneRefRotations;
 	TArray<FVector3f> BonePositions;
-	TArray<FVector4>  BoneRotations;
+	TArray<FVector4f> BoneRotations;
 	
 	if (DataAsset->Mode == EAnimToTextureMode::Bone)
 	{
 		// Gets Ref Bone Position and Rotations.
-		NumBones = GetRefBonePositionsAndRotations(DataAsset->GetSkeletalMesh(),
+		DataAsset->NumBones = GetRefBonePositionsAndRotations(DataAsset->GetSkeletalMesh(),
 			BoneRefPositions, BoneRefRotations);
 
 		// Add RefPose 
@@ -169,7 +168,6 @@ bool UAnimToTextureBPLibrary::AnimationToTexture(UAnimToTextureDataAsset* DataAs
 					
 				VertexDeltas.Append(VertexFrameDeltas);
 				VertexNormals.Append(VertexFrameNormals);
-
 			}
 
 			// ---------------------------------------------------------------------------
@@ -178,7 +176,7 @@ bool UAnimToTextureBPLibrary::AnimationToTexture(UAnimToTextureDataAsset* DataAs
 			else if (DataAsset->Mode == EAnimToTextureMode::Bone)
 			{
 				TArray<FVector3f> BoneFramePositions;
-				TArray<FVector4> BoneFrameRotations;
+				TArray<FVector4f> BoneFrameRotations;
 
 				GetBonePositionsAndRotations(SkeletalMeshComponent, BoneRefPositions,
 					BoneFramePositions, BoneFrameRotations);
@@ -243,7 +241,7 @@ bool UAnimToTextureBPLibrary::AnimationToTexture(UAnimToTextureDataAsset* DataAs
 		CreateUVChannel(DataAsset->GetStaticMesh(), DataAsset->StaticLODIndex, DataAsset->UVChannel, Height, Width);
 
 		// Update Bounds
-		SetBoundsExtensions(DataAsset->GetStaticMesh(), DataAsset->VertexMinBBox, DataAsset->VertexSizeBBox);
+		SetBoundsExtensions(DataAsset->GetStaticMesh(), (FVector)DataAsset->VertexMinBBox, (FVector)DataAsset->VertexSizeBBox);
 
 		// Done with StaticMesh
 		DataAsset->GetStaticMesh()->PostEditChange();
@@ -259,7 +257,7 @@ bool UAnimToTextureBPLibrary::AnimationToTexture(UAnimToTextureDataAsset* DataAs
 		// Write Bone Position and Rotation Textures
 		{
 			// Note we are adding +1 frame for the ref pose
-			if (!FindBestResolution(DataAsset->NumFrames + 1, NumBones,
+			if (!FindBestResolution(DataAsset->NumFrames + 1, DataAsset->NumBones,
 				Height, Width, DataAsset->BoneRowsPerFrame,
 				DataAsset->MaxHeight, DataAsset->MaxWidth, DataAsset->bEnforcePowerOfTwo))
 			{
@@ -269,7 +267,7 @@ bool UAnimToTextureBPLibrary::AnimationToTexture(UAnimToTextureDataAsset* DataAs
 
 			// Normalize Bone Data
 			TArray<FVector3f> NormalizedBonePositions;
-			TArray<FVector4> NormalizedBoneRotations;
+			TArray<FVector4f> NormalizedBoneRotations;
 			NormalizeBoneData(
 				BonePositions, BoneRotations,
 				DataAsset->BoneMinBBox, DataAsset->BoneSizeBBox,
@@ -279,16 +277,16 @@ bool UAnimToTextureBPLibrary::AnimationToTexture(UAnimToTextureDataAsset* DataAs
 			if (DataAsset->Precision == EAnimToTexturePrecision::SixteenBits)
 			{
 				WriteVectorsToTexture<FVector3f, FHighPrecision>(NormalizedBonePositions, DataAsset->NumFrames + 1, DataAsset->BoneRowsPerFrame, Height, Width, DataAsset->GetBonePositionTexture());
-				WriteVectorsToTexture<FVector4, FHighPrecision>(NormalizedBoneRotations, DataAsset->NumFrames + 1, DataAsset->BoneRowsPerFrame, Height, Width, DataAsset->GetBoneRotationTexture());
+				WriteVectorsToTexture<FVector4f, FHighPrecision>(NormalizedBoneRotations, DataAsset->NumFrames + 1, DataAsset->BoneRowsPerFrame, Height, Width, DataAsset->GetBoneRotationTexture());
 			}
 			else
 			{
 				WriteVectorsToTexture<FVector3f, FLowPrecision>(NormalizedBonePositions, DataAsset->NumFrames + 1, DataAsset->BoneRowsPerFrame, Height, Width, DataAsset->GetBonePositionTexture());
-				WriteVectorsToTexture<FVector4, FLowPrecision>(NormalizedBoneRotations, DataAsset->NumFrames + 1, DataAsset->BoneRowsPerFrame, Height, Width, DataAsset->GetBoneRotationTexture());
+				WriteVectorsToTexture<FVector4f, FLowPrecision>(NormalizedBoneRotations, DataAsset->NumFrames + 1, DataAsset->BoneRowsPerFrame, Height, Width, DataAsset->GetBoneRotationTexture());
 			}
 
 			// Update Bounds
-			SetBoundsExtensions(DataAsset->GetStaticMesh(), DataAsset->BoneMinBBox, DataAsset->BoneSizeBBox);
+			SetBoundsExtensions(DataAsset->GetStaticMesh(), (FVector)DataAsset->BoneMinBBox, (FVector)DataAsset->BoneSizeBBox);
 		}
 
 		// ---------------------------------------------------------------------------
@@ -308,7 +306,7 @@ bool UAnimToTextureBPLibrary::AnimationToTexture(UAnimToTextureDataAsset* DataAs
 
 			// Reduce BoneWeights to 4 Influences.
 			if (SocketIndex == INDEX_NONE)
-			{	
+			{
 				// Project SkinWeights from SkeletalMesh to StaticMesh
 				TArray<VertexSkinWeightMax> StaticMeshSkinWeights;
 				Mapping.ProjectSkinWeights(StaticMeshSkinWeights);
@@ -329,8 +327,16 @@ bool UAnimToTextureBPLibrary::AnimationToTexture(UAnimToTextureDataAsset* DataAs
 			}
 
 			// Write Bone Weights Texture
-			WriteSkinWeightsToTexture(SkinWeights,
-				DataAsset->BoneWeightRowsPerFrame, Height, Width, DataAsset->GetBoneWeightTexture());
+			if (DataAsset->Precision == EAnimToTexturePrecision::SixteenBits)
+			{
+				WriteSkinWeightsToTexture<FHighPrecision>(SkinWeights, DataAsset->NumBones,
+					DataAsset->BoneWeightRowsPerFrame, Height, Width, DataAsset->GetBoneWeightTexture());
+			}
+			else
+			{
+				WriteSkinWeightsToTexture<FLowPrecision>(SkinWeights, DataAsset->NumBones,
+					DataAsset->BoneWeightRowsPerFrame, Height, Width, DataAsset->GetBoneWeightTexture());
+			}
 
 			// Add Vertex UVChannel
 			CreateUVChannel(DataAsset->GetStaticMesh(), DataAsset->StaticLODIndex, DataAsset->UVChannel, Height, Width);
@@ -413,7 +419,7 @@ bool UAnimToTextureBPLibrary::CheckDataAsset(const UAnimToTextureDataAsset* Data
 	if (SourceModel.BuildSettings.bGenerateLightmapUVs &&
 		SourceModel.BuildSettings.DstLightmapIndex == DataAsset->UVChannel)
 	{
-		UE_LOG(LogAnimToTextureEditor, Warning, TEXT("Invalid UVChannel: %i. Already used by LightMap"), DataAsset->UVChannel);
+		UE_LOG(LogAnimToTextureEditor, Warning, TEXT("Invalid StaticMesh UVChannel: %i. Already used by LightMap"), DataAsset->UVChannel);
 		return false;
 	}
 
@@ -422,10 +428,10 @@ bool UAnimToTextureBPLibrary::CheckDataAsset(const UAnimToTextureDataAsset* Data
 	if (DataAsset->Precision == EAnimToTexturePrecision::EightBits &&
 		NumBones > 256)
 	{
-		UE_LOG(LogAnimToTextureEditor, Warning, TEXT("Invalid Number of Bones. There is a maximum of 256 bones for 8bit Precision"))
+		UE_LOG(LogAnimToTextureEditor, Warning, TEXT("Too many Bones: %i. There is a maximum of 256 bones for 8bit Precision"), NumBones);
 		return false;
 	}
-
+	
 	// Check Animations
 	OutAnimSequences.Reset();
 	for (const FAnimToTextureAnimSequenceInfo& AnimSequenceInfo : DataAsset->AnimSequences)
@@ -456,6 +462,12 @@ bool UAnimToTextureBPLibrary::CheckDataAsset(const UAnimToTextureDataAsset* Data
 		}
 	}
 
+	if (!OutAnimSequences.Num())
+	{
+		UE_LOG(LogAnimToTextureEditor, Warning, TEXT("No valid AnimSequences found"));
+		return false;
+	}
+
 	// All Good !
 	return true;
 }
@@ -529,7 +541,7 @@ void UAnimToTextureBPLibrary::GetVertexDeltasAndNormals(const USkeletalMeshCompo
 
 
 int32 UAnimToTextureBPLibrary::GetRefBonePositionsAndRotations(const USkeletalMesh* SkeletalMesh,
-	TArray<FVector3f>& OutBoneRefPositions, TArray<FVector4>& OutBoneRefRotations)
+	TArray<FVector3f>& OutBoneRefPositions, TArray<FVector4f>& OutBoneRefRotations)
 {
 	check(SkeletalMesh);
 
@@ -549,7 +561,7 @@ int32 UAnimToTextureBPLibrary::GetRefBonePositionsAndRotations(const USkeletalMe
 
 
 int32 UAnimToTextureBPLibrary::GetBonePositionsAndRotations(const USkeletalMeshComponent* SkeletalMeshComponent, const TArray<FVector3f>& BoneRefPositions,
-	TArray<FVector3f>& BonePositions, TArray<FVector4>& BoneRotations)
+	TArray<FVector3f>& BonePositions, TArray<FVector4f>& BoneRotations)
 {
 	check(SkeletalMeshComponent);
 
@@ -579,7 +591,7 @@ int32 UAnimToTextureBPLibrary::GetBonePositionsAndRotations(const USkeletalMeshC
 		// Decompose Transformation (ComponentSpace)
 		const FTransform& CompSpaceTransform = CompSpaceTransforms[BoneIndex];
 		FVector3f BonePosition;
-		FVector4 BoneRotation;
+		FVector4f BoneRotation;
 		DecomposeTransformation(CompSpaceTransform, BonePosition, BoneRotation);
 
 		// Position Delta (from RefPose)
@@ -587,7 +599,7 @@ int32 UAnimToTextureBPLibrary::GetBonePositionsAndRotations(const USkeletalMeshC
 
 		// Decompose Transformation (Relative to RefPose)
 		FVector3f BoneRelativePosition;
-		FVector4 BoneRelativeRotation;
+		FVector4f BoneRelativeRotation;
 		const FMatrix RefToLocalMatrix(RefToLocals[BoneIndex]);
 		const FTransform RelativeTransform(RefToLocalMatrix);
 		DecomposeTransformation(RelativeTransform, BoneRelativePosition, BoneRelativeRotation);
@@ -644,31 +656,21 @@ void UAnimToTextureBPLibrary::UpdateMaterialInstanceFromDataAsset(const UAnimToT
 	// Update Vertex Params
 	if (DataAsset->Mode == EAnimToTextureMode::Vertex)
 	{
-		FLinearColor VectorParameter;
-		VectorParameter = FLinearColor(DataAsset->VertexMinBBox);
-		UMaterialEditingLibrary::SetMaterialInstanceVectorParameterValue(MaterialInstance, AnimToTextureParamNames::BoundingBoxMin, VectorParameter, MaterialParameterAssociation);
-		
-		VectorParameter = FLinearColor(DataAsset->VertexSizeBBox);
-		UMaterialEditingLibrary::SetMaterialInstanceVectorParameterValue(MaterialInstance, AnimToTextureParamNames::BoundingBoxScale, VectorParameter, MaterialParameterAssociation);
+		UMaterialEditingLibrary::SetMaterialInstanceVectorParameterValue(MaterialInstance, AnimToTextureParamNames::MinBBox, FLinearColor(DataAsset->VertexMinBBox), MaterialParameterAssociation);
+		UMaterialEditingLibrary::SetMaterialInstanceVectorParameterValue(MaterialInstance, AnimToTextureParamNames::SizeBBox, FLinearColor(DataAsset->VertexSizeBBox), MaterialParameterAssociation);
 		UMaterialEditingLibrary::SetMaterialInstanceScalarParameterValue(MaterialInstance, AnimToTextureParamNames::RowsPerFrame, DataAsset->VertexRowsPerFrame, MaterialParameterAssociation);
-
 		UMaterialEditingLibrary::SetMaterialInstanceTextureParameterValue(MaterialInstance, AnimToTextureParamNames::VertexPositionTexture, DataAsset->GetVertexPositionTexture(), MaterialParameterAssociation);
 		UMaterialEditingLibrary::SetMaterialInstanceTextureParameterValue(MaterialInstance, AnimToTextureParamNames::VertexNormalTexture, DataAsset->GetVertexNormalTexture(), MaterialParameterAssociation);
-
 	}
 
 	// Update Bone Params
 	else if (DataAsset->Mode == EAnimToTextureMode::Bone)
 	{
-		FLinearColor VectorParameter;
-		VectorParameter = FLinearColor(DataAsset->BoneMinBBox);
-		UMaterialEditingLibrary::SetMaterialInstanceVectorParameterValue(MaterialInstance, AnimToTextureParamNames::BoundingBoxMin, VectorParameter, MaterialParameterAssociation);
-
-		VectorParameter = FLinearColor(DataAsset->BoneSizeBBox);
-		UMaterialEditingLibrary::SetMaterialInstanceVectorParameterValue(MaterialInstance, AnimToTextureParamNames::BoundingBoxScale, VectorParameter, MaterialParameterAssociation);
+		UMaterialEditingLibrary::SetMaterialInstanceScalarParameterValue(MaterialInstance, AnimToTextureParamNames::NumBones, DataAsset->NumBones, MaterialParameterAssociation);
+		UMaterialEditingLibrary::SetMaterialInstanceVectorParameterValue(MaterialInstance, AnimToTextureParamNames::MinBBox, FLinearColor(DataAsset->BoneMinBBox), MaterialParameterAssociation);
+		UMaterialEditingLibrary::SetMaterialInstanceVectorParameterValue(MaterialInstance, AnimToTextureParamNames::SizeBBox, FLinearColor(DataAsset->BoneSizeBBox), MaterialParameterAssociation);
 		UMaterialEditingLibrary::SetMaterialInstanceScalarParameterValue(MaterialInstance, AnimToTextureParamNames::RowsPerFrame, DataAsset->BoneRowsPerFrame, MaterialParameterAssociation);
 		UMaterialEditingLibrary::SetMaterialInstanceScalarParameterValue(MaterialInstance, AnimToTextureParamNames::BoneWeightRowsPerFrame, DataAsset->BoneWeightRowsPerFrame, MaterialParameterAssociation);
-
 		UMaterialEditingLibrary::SetMaterialInstanceTextureParameterValue(MaterialInstance, AnimToTextureParamNames::BonePositionTexture, DataAsset->GetBonePositionTexture(), MaterialParameterAssociation);
 		UMaterialEditingLibrary::SetMaterialInstanceTextureParameterValue(MaterialInstance, AnimToTextureParamNames::BoneRotationTexture, DataAsset->GetBoneRotationTexture(), MaterialParameterAssociation);
 		UMaterialEditingLibrary::SetMaterialInstanceTextureParameterValue(MaterialInstance, AnimToTextureParamNames::BoneWeightsTexture, DataAsset->GetBoneWeightTexture(), MaterialParameterAssociation);
@@ -887,7 +889,7 @@ UStaticMesh* UAnimToTextureBPLibrary::ConvertSkeletalMeshToStaticMesh(USkeletalM
 
 void UAnimToTextureBPLibrary::NormalizeVertexData(
 	const TArray<FVector3f>& Deltas, const TArray<FVector3f>& Normals,
-	FVector& OutMinBBox, FVector& OutSizeBBox,
+	FVector3f& OutMinBBox, FVector3f& OutSizeBBox,
 	TArray<FVector3f>& OutNormalizedDeltas, TArray<FVector3f>& OutNormalizedNormals)
 {
 	check(Deltas.Num() == Normals.Num());
@@ -910,14 +912,14 @@ void UAnimToTextureBPLibrary::NormalizeVertexData(
 		MaxBBox.Z = FMath::Max(Delta.Z, MaxBBox.Z);
 	}
 
-	OutSizeBBox = (FVector)MaxBBox - OutMinBBox;
+	OutSizeBBox = MaxBBox - OutMinBBox;
 
 	// ---------------------------------------------------------------------------
 	// Normalize Vertex Position Deltas
 	// Basically we want all deltas to be between [0, 1]
 	
 	// Compute Normalization Factor per-axis.
-	const FVector NormFactor = {
+	const FVector3f NormFactor = {
 		1.f / static_cast<float>(OutSizeBBox.X),
 		1.f / static_cast<float>(OutSizeBBox.Y),
 		1.f / static_cast<float>(OutSizeBBox.Z) };
@@ -925,7 +927,7 @@ void UAnimToTextureBPLibrary::NormalizeVertexData(
 	OutNormalizedDeltas.SetNumUninitialized(Deltas.Num());
 	for (int32 Index = 0; Index < Deltas.Num(); ++Index)
 	{
-		OutNormalizedDeltas[Index] = (FVector3f)(((FVector)Deltas[Index] - OutMinBBox) * NormFactor);
+		OutNormalizedDeltas[Index] = (Deltas[Index] - OutMinBBox) * NormFactor;
 	}
 
 	// ---------------------------------------------------------------------------
@@ -941,9 +943,9 @@ void UAnimToTextureBPLibrary::NormalizeVertexData(
 }
 
 void UAnimToTextureBPLibrary::NormalizeBoneData(
-	const TArray<FVector3f>& Positions, const TArray<FVector4>& Rotations,
-	FVector& OutMinBBox, FVector& OutSizeBBox, 
-	TArray<FVector3f>& OutNormalizedPositions, TArray<FVector4>& OutNormalizedRotations)
+	const TArray<FVector3f>& Positions, const TArray<FVector4f>& Rotations,
+	FVector3f& OutMinBBox, FVector3f& OutSizeBBox, 
+	TArray<FVector3f>& OutNormalizedPositions, TArray<FVector4f>& OutNormalizedRotations)
 {
 	check(Positions.Num() == Rotations.Num());
 
@@ -965,14 +967,14 @@ void UAnimToTextureBPLibrary::NormalizeBoneData(
 		MaxBBox.Z = FMath::Max(Position.Z, MaxBBox.Z);
 	}
 
-	OutSizeBBox = (FVector)MaxBBox - OutMinBBox;
+	OutSizeBBox = MaxBBox - OutMinBBox;
 
 	// ---------------------------------------------------------------------------
 	// Normalize Bone Position.
 	// Basically we want all positions to be between [0, 1]
 
 	// Compute Normalization Factor per-axis.
-	const FVector NormFactor = {
+	const FVector3f NormFactor = {
 		1.f / static_cast<float>(OutSizeBBox.X),
 		1.f / static_cast<float>(OutSizeBBox.Y),
 		1.f / static_cast<float>(OutSizeBBox.Z) };
@@ -980,7 +982,7 @@ void UAnimToTextureBPLibrary::NormalizeBoneData(
 	OutNormalizedPositions.SetNumUninitialized(Positions.Num());
 	for (int32 Index = 0; Index < Positions.Num(); ++Index)
 	{
-		OutNormalizedPositions[Index] = FVector3f(((FVector)Positions[Index] - OutMinBBox) * NormFactor);
+		OutNormalizedPositions[Index] = (Positions[Index] - OutMinBBox) * NormFactor;
 	}
 
 	// ---------------------------------------------------------------------------
@@ -989,15 +991,14 @@ void UAnimToTextureBPLibrary::NormalizeBoneData(
 	OutNormalizedRotations.SetNumUninitialized(Rotations.Num());
 	for (int32 Index = 0; Index < Rotations.Num(); ++Index)
 	{
-		const FVector4 Axis = Rotations[Index];
+		const FVector4f Axis = Rotations[Index];
 		const float Angle = Rotations[Index].W; // Angle are returned in radians and they go from [0-pi*2]
 
-		OutNormalizedRotations[Index] = (Axis.GetSafeNormal() + FVector::OneVector) * 0.5f;
+		OutNormalizedRotations[Index] = (Axis.GetSafeNormal() + FVector3f::OneVector) * 0.5f;
 		OutNormalizedRotations[Index].W = Angle / (PI * 2.f);
 	}
 }
 
-
 bool UAnimToTextureBPLibrary::CreateUVChannel(
 	UStaticMesh* StaticMesh, const int32 LODIndex, const int32 UVChannelIndex,
 	const int32 Height, const int32 Width)
diff --git a/Engine/Plugins/Experimental/AnimToTexture/Source/AnimToTextureEditor/Private/AnimToTextureUtils.cpp b/Engine/Plugins/Experimental/AnimToTexture/Source/AnimToTextureEditor/Private/AnimToTextureUtils.cpp
--- a/Engine/Plugins/Experimental/AnimToTexture/Source/AnimToTextureEditor/Private/AnimToTextureUtils.cpp
+++ b/Engine/Plugins/Experimental/AnimToTexture/Source/AnimToTextureEditor/Private/AnimToTextureUtils.cpp
@@ -6,23 +6,23 @@ namespace AnimToTexture_Private
 {
 
 void DecomposeTransformation(const FTransform& Transform, 
-	FVector3f& OutTranslation, FVector4& OutRotation)
+	FVector3f& OutTranslation, FVector4f& OutRotation)
 {
 	// Get Translation
 	OutTranslation = (FVector3f)Transform.GetTranslation();
 
 	// Get Rotation 
-	const FQuat Quat = Transform.GetRotation();
+	const FQuat4f Quat = (FQuat4f)Transform.GetRotation();
 
-	FVector Axis;
+	FVector3f Axis;
 	float Angle;
 	Quat.ToAxisAndAngle(Axis, Angle);
 
-	OutRotation = FVector4(Axis, Angle);
+	OutRotation = FVector4f(Axis, Angle);
 }
 
 void DecomposeTransformations(const TArray<FTransform>& Transforms, 
-	TArray<FVector3f>& OutTranslations, TArray<FVector4>& OutRotations)
+	TArray<FVector3f>& OutTranslations, TArray<FVector4f>& OutRotations)
 {
 	const int32 NumTransforms = Transforms.Num();
 	OutTranslations.SetNumUninitialized(NumTransforms);
@@ -34,37 +34,4 @@ void DecomposeTransformations(const TArray<FTransform>& Transforms,
 	}
 };
 
-bool WriteSkinWeightsToTexture(const TArray<VertexSkinWeightFour>& SkinWeights,
-	const int32 RowsPerFrame, const int32 Height, const int32 Width, UTexture2D* Texture)
-{
-	check(Texture);
-
-	const int32 NumVertices = SkinWeights.Num();
-
-	// Allocate PixelData.
-	TArray<FLowPrecision::ColorType> Pixels;
-	Pixels.Init(FLowPrecision::ColorType::Black, Height * Width);
-
-	for (int32 VertexIndex = 0; VertexIndex < NumVertices; ++VertexIndex)
-	{
-		const VertexSkinWeightFour& VertexSkinWeight = SkinWeights[VertexIndex];
-
-		// Write Influence
-		// NOTE: we are assuming the bone index is under < 256
-		Pixels[VertexIndex].R = (uint8)VertexSkinWeight.MeshBoneIndices[0];
-		Pixels[VertexIndex].G = (uint8)VertexSkinWeight.MeshBoneIndices[1];
-		Pixels[VertexIndex].B = (uint8)VertexSkinWeight.MeshBoneIndices[2];
-		Pixels[VertexIndex].A = (uint8)VertexSkinWeight.MeshBoneIndices[3];
-			
-		// Write Weight
-		Pixels[Width * RowsPerFrame + VertexIndex].R = VertexSkinWeight.BoneWeights[0];
-		Pixels[Width * RowsPerFrame + VertexIndex].G = VertexSkinWeight.BoneWeights[1];
-		Pixels[Width * RowsPerFrame + VertexIndex].B = VertexSkinWeight.BoneWeights[2];
-		Pixels[Width * RowsPerFrame + VertexIndex].A = VertexSkinWeight.BoneWeights[3];
-	};
-
-	// Write to Texture
-	return WriteToTexture<FLowPrecision>(Texture, Height, Width, Pixels);
-}
-
 } // end namespace AnimToTexture_Private
\ No newline at end of file
diff --git a/Engine/Plugins/Experimental/AnimToTexture/Source/AnimToTextureEditor/Public/AnimToTextureBPLibrary.h b/Engine/Plugins/Experimental/AnimToTexture/Source/AnimToTextureEditor/Public/AnimToTextureBPLibrary.h
--- a/Engine/Plugins/Experimental/AnimToTexture/Source/AnimToTextureEditor/Public/AnimToTextureBPLibrary.h
+++ b/Engine/Plugins/Experimental/AnimToTexture/Source/AnimToTextureEditor/Public/AnimToTextureBPLibrary.h
@@ -71,24 +71,24 @@ private:
 
 	// Gets RefPose Bone Position and Rotations.
 	static int32 GetRefBonePositionsAndRotations(const USkeletalMesh* SkeletalMesh, 
-		TArray<FVector3f>& OutBoneRefPositions, TArray<FVector4>& OutBoneRefRotations);
+		TArray<FVector3f>& OutBoneRefPositions, TArray<FVector4f>& OutBoneRefRotations);
 
 	// Gets Bone Position and Rotations for Current Pose.	
 	// The BonePosition is returned relative to the RefPose
 	static int32 GetBonePositionsAndRotations(const USkeletalMeshComponent* SkeletalMeshComponent, const TArray<FVector3f>& BoneRefPositions,
-		TArray<FVector3f>& BonePositions, TArray<FVector4>& BoneRotations);
+		TArray<FVector3f>& BonePositions, TArray<FVector4f>& BoneRotations);
 
 	// Normalizes Deltas and Normals between [0-1] with Bounding Box
 	static void NormalizeVertexData(
 		const TArray<FVector3f>& Deltas, const TArray<FVector3f>& Normals,
-		FVector& OutMinBBox, FVector& OutSizeBBox,
+		FVector3f& OutMinBBox, FVector3f& OutSizeBBox,
 		TArray<FVector3f>& OutNormalizedDeltas, TArray<FVector3f>& OutNormalizedNormals);
 
 	// Normalizes Positions and Rotations between [0-1] with Bounding Box
 	static void NormalizeBoneData(
-		const TArray<FVector3f>& Positions, const TArray<FVector4>& Rotations,
-		FVector& OutMinBBox, FVector& OutSizeBBox,
-		TArray<FVector3f>& OutNormalizedPositions, TArray<FVector4>& OutNormalizedRotations);
+		const TArray<FVector3f>& Positions, const TArray<FVector4f>& Rotations,
+		FVector3f& OutMinBBox, FVector3f& OutSizeBBox,
+		TArray<FVector3f>& OutNormalizedPositions, TArray<FVector4f>& OutNormalizedRotations);
 
 	/* Returns best resolution for the given data. 
 	*  Returns false if data doesnt fit in the the max range */
diff --git a/Engine/Plugins/Experimental/AnimToTexture/Source/AnimToTextureEditor/Public/AnimToTextureUtils.h b/Engine/Plugins/Experimental/AnimToTexture/Source/AnimToTextureEditor/Public/AnimToTextureUtils.h
--- a/Engine/Plugins/Experimental/AnimToTexture/Source/AnimToTextureEditor/Public/AnimToTextureUtils.h
+++ b/Engine/Plugins/Experimental/AnimToTexture/Source/AnimToTextureEditor/Public/AnimToTextureUtils.h
@@ -8,7 +8,6 @@
 #include "Engine/Texture.h"
 #include "Engine/Texture2D.h"
 
-
 namespace AnimToTexture_Private
 {
 
@@ -26,7 +25,7 @@ struct FLowPrecision
 	static constexpr EPixelFormat PixelFormat = EPixelFormat::PF_B8G8R8A8;
 	static constexpr ETextureSourceFormat TextureSourceFormat = ETextureSourceFormat::TSF_BGRA8;
 	static constexpr TextureCompressionSettings CompressionSettings = TextureCompressionSettings::TC_VectorDisplacementmap;
-	static constexpr ColorType DefaultColor = FColor(0, 0, 0);		
+	static constexpr ColorType DefaultColor = FColor(0, 0, 0, 0);		
 };
 
 struct FHighPrecision
@@ -49,7 +48,8 @@ bool WriteVectorsToTexture(const TArray<V>& Vectors,
 /* Writes list of skinweights into texture.
 *  The SkinWeights data is already in uint8 & uint16 format, no need for normalizing it.
 */
-bool WriteSkinWeightsToTexture(const TArray<VertexSkinWeightFour>& SkinWeights,
+template<class TextureSettings>
+bool WriteSkinWeightsToTexture(const TArray<VertexSkinWeightFour>& SkinWeights, const int32 NumBones,
 	const int32 RowsPerFrame,
 	const int32 Height, const int32 Width,
 	UTexture2D* Texture);
@@ -58,17 +58,18 @@ bool WriteSkinWeightsToTexture(const TArray<VertexSkinWeightFour>& SkinWeights,
 template<class TextureSettings>
 bool WriteToTexture(UTexture2D* Texture, const uint32 Height, const uint32 Width, const TArray<typename TextureSettings::ColorType>& Data);
 
-template<class V /* FVector / FVector4 / FIntVector */, class C /* FColor / FLinearColor */>
+template<class V /* FVector3f / FVector4f */, class C /* FColor / FVector4u16 */>
 void VectorToColor(const V& Vector, C& Color);
 
 /* Decomposes Transform in Translation and AxisAndAngle */
-void DecomposeTransformation(const FTransform& Transform, FVector3f& OutTranslation, FVector4& OutRotation);
-void DecomposeTransformations(const TArray<FTransform>& Transforms, TArray<FVector3f>& OutTranslations, TArray<FVector4>& OutRotations);
+void DecomposeTransformation(const FTransform& Transform, FVector3f& OutTranslation, FVector4f& OutRotation);
+void DecomposeTransformations(const TArray<FTransform>& Transforms, TArray<FVector3f>& OutTranslations, TArray<FVector4f>& OutRotations);
 
 } // end namespace AnimToTexture_Private
 
 // ----------------------------------------------------------------------------
 
+// LowPrecision
 template<>
 FORCEINLINE void AnimToTexture_Private::VectorToColor(const FVector3f& Vector, FColor& Color)
 {
@@ -79,18 +80,12 @@ FORCEINLINE void AnimToTexture_Private::VectorToColor(const FVector3f& Vector, F
 	Color.R = (uint8)FMath::RoundToInt(ClampedX * 255.f);
 	Color.G = (uint8)FMath::RoundToInt(ClampedY * 255.f);
 	Color.B = (uint8)FMath::RoundToInt(ClampedZ * 255.f);
-	Color.A = 255;
-}
-
-template<>
-FORCEINLINE void AnimToTexture_Private::VectorToColor(const FVector3d& Vector, FColor& Color)
-{
-	VectorToColor((FVector3f)Vector, Color);
+	Color.A = TNumericLimits<uint8>::Max();
 }
 
-
+// LowPrecision
 template<>
-FORCEINLINE void AnimToTexture_Private::VectorToColor(const FVector4& Vector, FColor& Color)
+FORCEINLINE void AnimToTexture_Private::VectorToColor(const FVector4f& Vector, FColor& Color)
 {
 	const float ClampedX = FMath::Clamp(Vector.X, 0.f, 1.f);
 	const float ClampedY = FMath::Clamp(Vector.Y, 0.f, 1.f);
@@ -103,72 +98,28 @@ FORCEINLINE void AnimToTexture_Private::VectorToColor(const FVector4& Vector, FC
 	Color.A = (uint8)FMath::RoundToInt(ClampedW * 255.f);
 }
 
-template<>
-FORCEINLINE void AnimToTexture_Private::VectorToColor(const FIntVector& Vector, FColor& Color)
-{
-	const int32 ClampedX = FMath::Clamp(Vector.X, 0, 1);
-	const int32 ClampedY = FMath::Clamp(Vector.Y, 0, 1);
-	const int32 ClampedZ = FMath::Clamp(Vector.Z, 0, 1);
-
-	Color.R = (uint8)(ClampedX * 255);
-	Color.G = (uint8)(ClampedY * 255);
-	Color.B = (uint8)(ClampedZ * 255);
-	Color.A = 255;
-}
-
-template<>
-FORCEINLINE void AnimToTexture_Private::VectorToColor(const FIntVector4& Vector, FColor& Color)
-{
-	const int32 ClampedX = FMath::Clamp(Vector.X, 0, 1);
-	const int32 ClampedY = FMath::Clamp(Vector.Y, 0, 1);
-	const int32 ClampedZ = FMath::Clamp(Vector.Z, 0, 1);
-	const int32 ClampedW = FMath::Clamp(Vector.W, 0, 1);
-
-	Color.R = (uint8)(ClampedX * 255);
-	Color.G = (uint8)(ClampedY * 255);
-	Color.B = (uint8)(ClampedZ * 255);
-	Color.A = (uint8)(ClampedW * 255);
-}
-
-template<>
-FORCEINLINE void AnimToTexture_Private::VectorToColor(const FVector3f& Vector, FLinearColor& Color)
-{
-	Color.R = Vector.X;
-	Color.G = Vector.Y;
-	Color.B = Vector.Z;
-	Color.A = 1.f;
-}
-
-template<>
-FORCEINLINE void AnimToTexture_Private::VectorToColor(const FVector4& Vector, FLinearColor& Color)
-{
-	Color.R = Vector.X;
-	Color.G = Vector.Y;
-	Color.B = Vector.Z;
-	Color.A = Vector.W;
-}
-
+// HighPrecision
 template<>
 FORCEINLINE void AnimToTexture_Private::VectorToColor(const FVector3f& Vector, FVector4u16& Color)
 {
-	Color.X = FMath::RoundToInt(FMath::Clamp(Vector.X, 0.f, 1.f) * MAX_uint16);
-	Color.Y = FMath::RoundToInt(FMath::Clamp(Vector.Y, 0.f, 1.f) * MAX_uint16);
-	Color.Z = FMath::RoundToInt(FMath::Clamp(Vector.Z, 0.f, 1.f) * MAX_uint16);
-	Color.W = MAX_uint16;
+	Color.X = FMath::RoundToInt(FMath::Clamp(Vector.X, 0.f, 1.f) * TNumericLimits<uint16>::Max());
+	Color.Y = FMath::RoundToInt(FMath::Clamp(Vector.Y, 0.f, 1.f) * TNumericLimits<uint16>::Max());
+	Color.Z = FMath::RoundToInt(FMath::Clamp(Vector.Z, 0.f, 1.f) * TNumericLimits<uint16>::Max());
+	Color.W = TNumericLimits<uint16>::Max();
 }
 
+// HighPrecision
 template<>
-FORCEINLINE void AnimToTexture_Private::VectorToColor(const FVector4& Vector, FVector4u16& Color)
+FORCEINLINE void AnimToTexture_Private::VectorToColor(const FVector4f& Vector, FVector4u16& Color)
 {
-	Color.X = FMath::RoundToInt(FMath::Clamp(Vector.X, 0.f, 1.f) * MAX_uint16);
-	Color.Y = FMath::RoundToInt(FMath::Clamp(Vector.Y, 0.f, 1.f) * MAX_uint16);
-	Color.Z = FMath::RoundToInt(FMath::Clamp(Vector.Z, 0.f, 1.f) * MAX_uint16);
-	Color.W = FMath::RoundToInt(FMath::Clamp(Vector.W, 0.f, 1.f) * MAX_uint16);
+	Color.X = FMath::RoundToInt(FMath::Clamp(Vector.X, 0.f, 1.f) * TNumericLimits<uint16>::Max());
+	Color.Y = FMath::RoundToInt(FMath::Clamp(Vector.Y, 0.f, 1.f) * TNumericLimits<uint16>::Max());
+	Color.Z = FMath::RoundToInt(FMath::Clamp(Vector.Z, 0.f, 1.f) * TNumericLimits<uint16>::Max());
+	Color.W = FMath::RoundToInt(FMath::Clamp(Vector.W, 0.f, 1.f) * TNumericLimits<uint16>::Max());
 }
 
-
 template<class V, class TextureSettings>
-FORCEINLINE bool AnimToTexture_Private::WriteVectorsToTexture(const TArray<V>& Vectors,
+FORCEINLINE_DEBUGGABLE bool AnimToTexture_Private::WriteVectorsToTexture(const TArray<V>& Vectors,
 	const int32 NumFrames, const int32 RowsPerFrame,
 	const int32 Height, const int32 Width, UTexture2D* Texture)
 {
@@ -204,15 +155,59 @@ FORCEINLINE bool AnimToTexture_Private::WriteVectorsToTexture(const TArray<V>& V
 }
 
 template<class TextureSettings>
-FORCEINLINE bool AnimToTexture_Private::WriteToTexture(
+FORCEINLINE_DEBUGGABLE bool AnimToTexture_Private::WriteSkinWeightsToTexture(const TArray<AnimToTexture_Private::VertexSkinWeightFour>& SkinWeights, const int32 NumBones,
+	const int32 RowsPerFrame, const int32 Height, const int32 Width, UTexture2D* Texture)
+{
+	check(Texture);
+	
+	const int32 NumVertices = SkinWeights.Num();
+
+	// Allocate PixelData.
+	TArray<typename TextureSettings::ColorType> Pixels;
+	Pixels.Init(TextureSettings::DefaultColor, Height * Width);
+
+	for (int32 VertexIndex = 0; VertexIndex < NumVertices; ++VertexIndex)
+	{
+		const VertexSkinWeightFour& VertexSkinWeight = SkinWeights[VertexIndex];
+
+		// Normalize BoneIndices
+		const FVector4f BoneIndices(
+			(float)VertexSkinWeight.MeshBoneIndices[0] / float(NumBones),
+			(float)VertexSkinWeight.MeshBoneIndices[1] / float(NumBones),
+			(float)VertexSkinWeight.MeshBoneIndices[2] / float(NumBones),
+			(float)VertexSkinWeight.MeshBoneIndices[3] / float(NumBones));
+
+		// Normalize BoneWeights
+		const FVector4f BoneWeights(
+			(float)VertexSkinWeight.BoneWeights[0] / 255.f,
+			(float)VertexSkinWeight.BoneWeights[1] / 255.f,
+			(float)VertexSkinWeight.BoneWeights[2] / 255.f,
+			(float)VertexSkinWeight.BoneWeights[3] / 255.f);
+
+		// Write BoneIndex
+		{
+			typename TextureSettings::ColorType& Pixel = Pixels[VertexIndex];
+			VectorToColor<FVector4f, typename TextureSettings::ColorType>(BoneIndices, Pixel);
+		}
+		
+		// Write BoneWeight
+		{
+			typename TextureSettings::ColorType& Pixel = Pixels[RowsPerFrame * Width + VertexIndex];
+			VectorToColor<FVector4f, typename TextureSettings::ColorType>(BoneWeights, Pixel);
+		}
+	};
+
+	// Write to Texture
+	return WriteToTexture<TextureSettings>(Texture, Height, Width, Pixels);
+}
+
+template<class TextureSettings>
+FORCEINLINE_DEBUGGABLE bool AnimToTexture_Private::WriteToTexture(
 	UTexture2D* Texture,
 	const uint32 Height, const uint32 Width,
 	const TArray<typename TextureSettings::ColorType>& Pixels)
 {
-	if (!Texture)
-	{
-		return false;
-	}
+	check(Texture);
 
 	// ------------------------------------------------------------------------
 	// Get Texture Platform
@@ -274,5 +269,3 @@ FORCEINLINE bool AnimToTexture_Private::WriteToTexture(
 
 	return true;
 }
-
-
