[AnimToTexture] Optimized Mapping.
- Added NumberOfDriverTriangles and Sigma params to DataAsset
-

diff --git a/Engine/Plugins/Experimental/AnimToTexture/Source/AnimToTexture/Public/AnimToTextureDataAsset.h b/Engine/Plugins/Experimental/AnimToTexture/Source/AnimToTexture/Public/AnimToTextureDataAsset.h
--- a/Engine/Plugins/Experimental/AnimToTexture/Source/AnimToTexture/Public/AnimToTextureDataAsset.h
+++ b/Engine/Plugins/Experimental/AnimToTexture/Source/AnimToTexture/Public/AnimToTextureDataAsset.h
@@ -143,6 +143,23 @@ public:
 	UPROPERTY(EditAnywhere, BlueprintReadWrite, Category = "StaticMesh")
 	int32 UVChannel = 1;
 
+	/**
+	* Number of Driver Triangles
+	* Each StaticMesh Vertex will be influenced by N SkeletalMesh (Driver) Triangles.
+	* Increasing the Number of Driver Triangles will increase the Mapping computation.
+	* Using a single Driver Triangle will do a Rigid Binding to Closest Triangle.
+	*/
+	UPROPERTY(EditAnywhere, BlueprintReadWrite, Category = "StaticMesh|Mapping")
+	int32 NumDriverTriangles = 10;
+
+	/**
+	* Inverse Distance Weighting
+	* This exponent value will be used for computing weights for the DriverTriangles.
+	* Larger number will create a more contrasted weighting, but it might 
+	*/
+	UPROPERTY(EditAnywhere, BlueprintReadWrite, Category = "StaticMesh|Mapping")
+	float Sigma = 1.f;
+
 	// ------------------------------------------------------
 	// Texture
 
diff --git a/Engine/Plugins/Experimental/AnimToTexture/Source/AnimToTextureEditor/Private/AnimToTextureBPLibrary.cpp b/Engine/Plugins/Experimental/AnimToTexture/Source/AnimToTextureEditor/Private/AnimToTextureBPLibrary.cpp
--- a/Engine/Plugins/Experimental/AnimToTexture/Source/AnimToTextureEditor/Private/AnimToTextureBPLibrary.cpp
+++ b/Engine/Plugins/Experimental/AnimToTexture/Source/AnimToTextureEditor/Private/AnimToTextureBPLibrary.cpp
@@ -58,7 +58,7 @@ bool UAnimToTextureBPLibrary::AnimationToTexture(UAnimToTextureDataAsset* DataAs
 		ProgressBar.MakeDialog(false /*bShowCancelButton*/, false /*bAllowInPIE*/);
 
 		Mapping.Update(DataAsset->GetStaticMesh(), DataAsset->StaticLODIndex,
-			DataAsset->GetSkeletalMesh(), DataAsset->SkeletalLODIndex);
+			DataAsset->GetSkeletalMesh(), DataAsset->SkeletalLODIndex, DataAsset->NumDriverTriangles, DataAsset->Sigma);
 	}
 
 	// Get Number of Source Vertices (StaticMesh)
diff --git a/Engine/Plugins/Experimental/AnimToTexture/Source/AnimToTextureEditor/Private/AnimToTextureMeshMapping.cpp b/Engine/Plugins/Experimental/AnimToTexture/Source/AnimToTextureEditor/Private/AnimToTextureMeshMapping.cpp
--- a/Engine/Plugins/Experimental/AnimToTexture/Source/AnimToTextureEditor/Private/AnimToTextureMeshMapping.cpp
+++ b/Engine/Plugins/Experimental/AnimToTexture/Source/AnimToTextureEditor/Private/AnimToTextureMeshMapping.cpp
@@ -7,34 +7,58 @@ namespace AnimToTexture_Private
 
 void FSourceVertexData::Update(const FVector3f& SourceVertex,
 	const TArray<FVector3f>& DriverVertices, const TArray<FIntVector3>& DriverTriangles, const TArray<VertexSkinWeightMax>& DriverSkinWeights, 
-	const float Sigma)
+	const int32 NumDrivers, const float Sigma)
 {	
 	const int32 NumDriverTriangles = DriverTriangles.Num();
+	const int32 NDriverTriangles = FMath::Clamp(NumDrivers, 1, NumDriverTriangles);
 
-	// Get ClosestPoint to Triangles
+	// Get Distances from Vertex to Each Triangle (ClosestPoint)
+	TArray<TPair<float, int32>> SortedDistances;
 	TArray<FVector3f> ClosestPoints;
-	ClosestPoints.SetNumUninitialized(NumDriverTriangles);
+	{
+		SortedDistances.SetNumUninitialized(NumDriverTriangles);
+		ClosestPoints.SetNumUninitialized(NumDriverTriangles);
+
+		for (int32 DriverTriangleIndex = 0; DriverTriangleIndex < NumDriverTriangles; DriverTriangleIndex++)
+		{
+			const FIntVector3& DriverTriangle = DriverTriangles[DriverTriangleIndex];
+			const FVector3f& A = DriverVertices[DriverTriangle.X];
+			const FVector3f& B = DriverVertices[DriverTriangle.Y];
+			const FVector3f& C = DriverVertices[DriverTriangle.Z];
+
+			// ClosestPoint To Triangle 
+			ClosestPoints[DriverTriangleIndex] = FindClosestPointToTriangle(SourceVertex, A, B, C);
 
-	for (int32 DriverTriangleIndex = 0; DriverTriangleIndex < NumDriverTriangles; DriverTriangleIndex++)
+			// Distance To Triangle
+			const float Distance = FVector3f::Distance(SourceVertex, ClosestPoints[DriverTriangleIndex]);
+			SortedDistances[DriverTriangleIndex] = TPair<float, int32>(Distance, DriverTriangleIndex);
+		};
+
+		// Sort By Distance
+		SortedDistances.Sort();
+	}
+
+	// Get Inverse Distance from Vertex to N-Closest Triangles
+	TArray<float> NWeights;
 	{
-		const FIntVector3& DriverTriangle = DriverTriangles[DriverTriangleIndex];
-		const FVector3f& A = DriverVertices[DriverTriangle.X];
-		const FVector3f& B = DriverVertices[DriverTriangle.Y];
-		const FVector3f& C = DriverVertices[DriverTriangle.Z];
+		TArray<FVector3f> NClosestPoints;
+		NClosestPoints.SetNumUninitialized(NDriverTriangles);
 
-		// ClosestPoint To Triangle 
-		int32 OnPointLocalIndex = INDEX_NONE;
-		ClosestPoints[DriverTriangleIndex] = FindClosestPointToTriangle(SourceVertex, A, B, C, OnPointLocalIndex);
-	};
+		for (int32 Index = 0; Index < NDriverTriangles; Index++)
+		{
+			const int32& DriverTriangleIndex = SortedDistances[Index].Value;
+			NClosestPoints[Index] = ClosestPoints[DriverTriangleIndex];
+		}
 
-	// Get Inverse Distance Weighting from Vertex to each DriverTriangle (ClosestPoint)
-	TArray<float> Weights;
-	AnimToTexture_Private::InverseDistanceWeights(SourceVertex, ClosestPoints, Weights, Sigma);
+		AnimToTexture_Private::InverseDistanceWeights(SourceVertex, NClosestPoints, NWeights, Sigma);
+	}
 
-	for (int32 DriverTriangleIndex = 0; DriverTriangleIndex < NumDriverTriangles; DriverTriangleIndex++)
+	DriverTriangleData.Reserve(NDriverTriangles);
+	for (int32 Index = 0; Index < NDriverTriangles; Index++)
 	{
-		// no need to store data if weight is small
-		if (Weights[DriverTriangleIndex] > UE_KINDA_SMALL_NUMBER)
+		const int32& DriverTriangleIndex = SortedDistances[Index].Value;
+
+		if (NWeights[Index] > UE_KINDA_SMALL_NUMBER)
 		{
 			const FVector3f& ClosestPoint = ClosestPoints[DriverTriangleIndex];
 			const FIntVector3& DriverTriangle = DriverTriangles[DriverTriangleIndex];
@@ -43,7 +67,7 @@ void FSourceVertexData::Update(const FVector3f& SourceVertex,
 			const FVector3f& C = DriverVertices[DriverTriangle.Z];
 
 			FSourceVertexDriverTriangleData TriangleData;
-			TriangleData.InverseDistanceWeight = Weights[DriverTriangleIndex];
+			TriangleData.InverseDistanceWeight = NWeights[Index];
 			TriangleData.Triangle = DriverTriangle;
 			TriangleData.BarycentricCoords = BarycentricCoordinates(ClosestPoint, A, B, C);
 			TriangleData.TangentLocalIndex = GetTriangleTangentLocalIndex(ClosestPoint, A, B, C);
@@ -61,7 +85,8 @@ void FSourceVertexData::Update(const FVector3f& SourceVertex,
 
 
 void FSourceMeshToDriverMesh::Update(const UStaticMesh* StaticMesh, const int32 StaticMeshLODIndex, 
-	const USkeletalMesh* SkeletalMesh, const int32 SkeletalMeshLODIndex)
+	const USkeletalMesh* SkeletalMesh, const int32 SkeletalMeshLODIndex, 
+	const int32 NumDrivers, const float Sigma)
 {
 	check(StaticMesh);
 	check(SkeletalMesh);
@@ -86,7 +111,7 @@ void FSourceMeshToDriverMesh::Update(const UStaticMesh* StaticMesh, const int32
 	{	
 		// Create Mapping from StaticMesh Vertex to SkeletalMesh Triangles
 		SourceVerticesData[SourceVertexIndex].Update(SourceVertices[SourceVertexIndex], 
-			DriverVertices, DriverTriangles, DriverSkinWeights, 1.f /*Sigma*/);
+			DriverVertices, DriverTriangles, DriverSkinWeights, NumDrivers, Sigma);
 
 		// UE_LOG(LogTemp, Warning, TEXT("Vertex: %i NumTriangles: %i."), SourceVertexIndex, SourceVerticesData[SourceVertexIndex].DriverTriangleData.Num());
 
diff --git a/Engine/Plugins/Experimental/AnimToTexture/Source/AnimToTextureEditor/Private/AnimToTextureSkeletalMesh.cpp b/Engine/Plugins/Experimental/AnimToTexture/Source/AnimToTextureEditor/Private/AnimToTextureSkeletalMesh.cpp
--- a/Engine/Plugins/Experimental/AnimToTexture/Source/AnimToTextureEditor/Private/AnimToTextureSkeletalMesh.cpp
+++ b/Engine/Plugins/Experimental/AnimToTexture/Source/AnimToTextureEditor/Private/AnimToTextureSkeletalMesh.cpp
@@ -394,8 +394,7 @@ void GetSkinnedVertices(const USkeletalMeshComponent* SkeletalMeshComponent, con
 	};
 };
 
-FVector3f FindClosestPointToTriangle(const FVector3f& P, const FVector3f& A, const FVector3f& B, const FVector3f& C, 
-	int32& OutOnPointLocalIndex)
+FVector3f FindClosestPointToTriangle(const FVector3f& P, const FVector3f& A, const FVector3f& B, const FVector3f& C)
 {
 	const FVector3f AB = B - A;
 	const FVector3f AC = C - A;
@@ -405,7 +404,6 @@ FVector3f FindClosestPointToTriangle(const FVector3f& P, const FVector3f& A, con
 	const float D2 = FVector3f::DotProduct(AC, AP);
 	if (D1 <= 0.f && D2 <= 0.f)
 	{
-		OutOnPointLocalIndex = 0;
 		return A;
 	}
 
@@ -414,7 +412,6 @@ FVector3f FindClosestPointToTriangle(const FVector3f& P, const FVector3f& A, con
 	const float D4 = FVector3f::DotProduct(AC, BP);
 	if (D3 >= 0.f && D4 <= D3)
 	{
-		OutOnPointLocalIndex = 1;
 		return B;
 	}
 
@@ -423,14 +420,12 @@ FVector3f FindClosestPointToTriangle(const FVector3f& P, const FVector3f& A, con
 	const float D6 = FVector3f::DotProduct(AC, CP);
 	if (D6 >= 0.f && D5 <= D6)
 	{
-		OutOnPointLocalIndex = 2;
 		return C;
 	}
 
 	const float VC = D1 * D4 - D3 * D2;
 	if (VC <= 0.f && D1 >= 0.f && D3 <= 0.f)
 	{
-		OutOnPointLocalIndex = INDEX_NONE;
 		const float V = D1 / (D1 - D3);
 		return A + V * AB;
 	}
@@ -438,7 +433,6 @@ FVector3f FindClosestPointToTriangle(const FVector3f& P, const FVector3f& A, con
 	const float VB = D5 * D2 - D1 * D6;
 	if (VB <= 0.f && D2 >= 0.f && D6 <= 0.f)
 	{
-		OutOnPointLocalIndex = INDEX_NONE;
 		const float V = D2 / (D2 - D6);
 		return A + V * AC;
 	}
@@ -446,12 +440,10 @@ FVector3f FindClosestPointToTriangle(const FVector3f& P, const FVector3f& A, con
 	const float VA = D3 * D6 - D5 * D4;
 	if (VA <= 0.f && (D4 - D3) >= 0.f && (D5 - D6) >= 0.f)
 	{
-		OutOnPointLocalIndex = INDEX_NONE;
 		const float V = (D4 - D3) / ((D4 - D3) + (D5 - D6));
 		return B + V * (C - B);
 	}
 
-	OutOnPointLocalIndex = INDEX_NONE;
 	const float Denom = 1.0f / (VA + VB + VC);
 	const float V = VB * Denom;
 	const float W = VC * Denom;
@@ -532,13 +524,14 @@ FVector3f BarycentricCoordinates(const FVector3f& P, const FVector3f& A, const F
 	return FVector3f(U, V, W);
 }
 
-int32 InverseDistanceWeights(const FVector3f& Point, const TArray<FVector3f>& Points,
-	TArray<float>& OutWeights, float Sigma)
+void InverseDistanceWeights(const FVector3f& Point, const TArray<FVector3f>& Points,
+	TArray<float>& OutWeights, const float Sigma)
 {
 	// Allocate
 	const int32 Count = Points.Num();
 	OutWeights.SetNumZeroed(Count);
 
+	const float SafeSigma = FMath::Max(UE_KINDA_SMALL_NUMBER, Sigma);
 	float SumInverseDistance = 0.f;
 	TArray<float> InverseDistances;
 	InverseDistances.SetNumUninitialized(Count);
@@ -551,14 +544,10 @@ int32 InverseDistanceWeights(const FVector3f& Point, const TArray<FVector3f>& Po
 		if (Distance < UE_KINDA_SMALL_NUMBER)
 		{
 			OutWeights[Index] = 1.f;
-			return Index;
+			return;
 		}
 
-		InverseDistances[Index] = 1.f / FMath::Pow(Distance, Sigma);
-		if (InverseDistances[Index] < UE_KINDA_SMALL_NUMBER)
-		{
-			InverseDistances[Index] = 0.f;
-		}
+		InverseDistances[Index] = 1.f / FMath::Pow(Distance, SafeSigma);
 		SumInverseDistance += InverseDistances[Index];
 	}
 
@@ -567,8 +556,6 @@ int32 InverseDistanceWeights(const FVector3f& Point, const TArray<FVector3f>& Po
 	{
 		OutWeights[Index] = InverseDistances[Index] / SumInverseDistance;
 	}
-
-	return INDEX_NONE;
 }
 
 
diff --git a/Engine/Plugins/Experimental/AnimToTexture/Source/AnimToTextureEditor/Public/AnimToTextureMeshMapping.h b/Engine/Plugins/Experimental/AnimToTexture/Source/AnimToTextureEditor/Public/AnimToTextureMeshMapping.h
--- a/Engine/Plugins/Experimental/AnimToTexture/Source/AnimToTextureEditor/Public/AnimToTextureMeshMapping.h
+++ b/Engine/Plugins/Experimental/AnimToTexture/Source/AnimToTextureEditor/Public/AnimToTextureMeshMapping.h
@@ -13,7 +13,6 @@ struct FSourceVertexDriverTriangleData
 	uint8               TangentLocalIndex;
 	float               InverseDistanceWeight;
 	FIntVector3         Triangle;
-	FVector3f           ClosestPoint;
 	FVector3f           BarycentricCoords;
 	FMatrix44f          InvMatrix;
 	VertexSkinWeightMax SkinWeights;
@@ -27,7 +26,7 @@ public:
 	
 	void Update(const FVector3f& SourceVertex,
 		const TArray<FVector3f>& DriverVertices, const TArray<FIntVector3>& DriverTriangles, const TArray<VertexSkinWeightMax>& DriverSkinWeights, 
-		const float Sigma=1.f);
+		const int32 NumDrivers, const float Sigma=1.f);
 
 	// DriverTriangle Data specific to this SourceVertex
 	TArray<FSourceVertexDriverTriangleData> DriverTriangleData;
@@ -42,7 +41,8 @@ public:
 	FSourceMeshToDriverMesh() = default;
 	
 	void Update(const UStaticMesh* StaticMesh, const int32 StaticMeshLODIndex,
-		const USkeletalMesh* SkeletalMesh, const int32 SkeletalMeshLODIndex);
+		const USkeletalMesh* SkeletalMesh, const int32 SkeletalMeshLODIndex, 
+		const int32 NumDrivers, const float Sigma=1.f);
 
 	// Returns Number of Source Vertices
 	int32 GetNumSourceVertices() const;
diff --git a/Engine/Plugins/Experimental/AnimToTexture/Source/AnimToTextureEditor/Public/AnimToTextureSkeletalMesh.h b/Engine/Plugins/Experimental/AnimToTexture/Source/AnimToTextureEditor/Public/AnimToTextureSkeletalMesh.h
--- a/Engine/Plugins/Experimental/AnimToTexture/Source/AnimToTextureEditor/Public/AnimToTextureSkeletalMesh.h
+++ b/Engine/Plugins/Experimental/AnimToTexture/Source/AnimToTextureEditor/Public/AnimToTextureSkeletalMesh.h
@@ -62,9 +62,8 @@ bool HasBone(const USkeletalMesh* SkeletalMesh, const FName& Bone);
    Only the RawBones are returned (no virtual bones) */
 void GetBoneNames(const USkeletalMesh* SkeletalMesh, TArray<FName>& OutNames);
 
-/* Computes closest point to triangle. Returns index of point [0-2] if right on top, else INDEX_NONE */
-FVector3f FindClosestPointToTriangle(const FVector3f& Point, const FVector3f& PointA, const FVector3f& PointB, const FVector3f& PointC, 
-	int32& OutOnPointLocalIndex);
+/* Computes closest point to triangle */
+FVector3f FindClosestPointToTriangle(const FVector3f& Point, const FVector3f& PointA, const FVector3f& PointB, const FVector3f& PointC);
 
 /* Computes Barycentric coordinates from point to triangle */
 FVector3f BarycentricCoordinates(const FVector3f& Point, const FVector3f& PointA, const FVector3f& PointB, const FVector3f& PointC);
@@ -86,8 +85,7 @@ FMatrix44f GetTriangleMatrix(const FVector3f& Point, const FVector3f& PointA, co
 
 // Returns a list of weights from Point to each of the given Points.
 // The returned Weights are the size of Points and normalized (sum to 1.0) 
-// If one of the points is right on top (no need to interpolate), it will return index, else INDEX_NONE
-int32 InverseDistanceWeights(const FVector3f& Point, const TArray<FVector3f>& Points,
-	TArray<float>& OutWeights, float Sigma = 1.f);
+void InverseDistanceWeights(const FVector3f& Point, const TArray<FVector3f>& Points,
+	TArray<float>& OutWeights, const float Sigma = 1.f);
 
 } // end namespace AnimToTexture_Private
