PR #13166: Exposing anim to texture editor functions as a public interface.

diff --git a/Engine/Plugins/Experimental/AnimToTexture/Source/AnimToTextureEditor/Public/AnimToTextureBPLibrary.h b/Engine/Plugins/Experimental/AnimToTexture/Source/AnimToTextureEditor/Public/AnimToTextureBPLibrary.h
--- a/Engine/Plugins/Experimental/AnimToTexture/Source/AnimToTextureEditor/Public/AnimToTextureBPLibrary.h
+++ b/Engine/Plugins/Experimental/AnimToTexture/Source/AnimToTextureEditor/Public/AnimToTextureBPLibrary.h
@@ -28,25 +28,25 @@ class UAnimToTextureBPLibrary : public UBlueprintFunctionLibrary
 	* Bakes Animation Data into Textures.
 	*/
 	UFUNCTION(BlueprintCallable, meta = (Category = "AnimToTexture"))
-	static UPARAM(DisplayName="bSuccess") bool AnimationToTexture(UAnimToTextureDataAsset* DataAsset);
+	static ANIMTOTEXTUREEDITOR_API UPARAM(DisplayName="bSuccess") bool AnimationToTexture(UAnimToTextureDataAsset* DataAsset);
 
 	/** 
 	* Utility for converting SkeletalMesh into a StaticMesh
 	*/
 	UFUNCTION(BlueprintCallable, Category = "AnimToTexture")
-	static UStaticMesh* ConvertSkeletalMeshToStaticMesh(USkeletalMesh* SkeletalMesh, const FString PackageName, const int32 LODIndex = -1);
+	static ANIMTOTEXTUREEDITOR_API UStaticMesh* ConvertSkeletalMeshToStaticMesh(USkeletalMesh* SkeletalMesh, const FString PackageName, const int32 LODIndex = -1);
 
 	/**
 	* Utility for setting a StaticMesh LightMapIndex.
 	*/
 	UFUNCTION(BlueprintCallable, Category = "AnimToTexture")
-	static UPARAM(DisplayName = "bSuccess") bool SetLightMapIndex(UStaticMesh* StaticMesh, const int32 LODIndex, const int32 LightmapIndex=1, bool bGenerateLightmapUVs=true);
+	static ANIMTOTEXTUREEDITOR_API UPARAM(DisplayName = "bSuccess") bool SetLightMapIndex(UStaticMesh* StaticMesh, const int32 LODIndex, const int32 LightmapIndex=1, bool bGenerateLightmapUVs=true);
 
 	/**
 	* Updates a material's parameters to match those of an AnimToTexture DataAsset
 	*/
 	UFUNCTION(BlueprintCallable, meta = (Category = "AnimToTexture"))
-	static void UpdateMaterialInstanceFromDataAsset(const UAnimToTextureDataAsset* DataAsset, class UMaterialInstanceConstant* MaterialInstance,
+	static ANIMTOTEXTUREEDITOR_API void UpdateMaterialInstanceFromDataAsset(const UAnimToTextureDataAsset* DataAsset, class UMaterialInstanceConstant* MaterialInstance,
 		const EMaterialParameterAssociation MaterialParameterAssociation = EMaterialParameterAssociation::LayerParameter);
 
 #endif // WITH_EDITOR
