[AnimToTexture] Fixed BatchUpdateInstancesFrameData update, now using FTransform instead of FMatrix.

#changelist validated

#virtualized

diff --git a/Engine/Plugins/Experimental/AnimToTexture/Source/AnimToTexture/Private/AnimToTextureInstancePlaybackHelpers.cpp b/Engine/Plugins/Experimental/AnimToTexture/Source/AnimToTexture/Private/AnimToTextureInstancePlaybackHelpers.cpp
--- a/Engine/Plugins/Experimental/AnimToTexture/Source/AnimToTexture/Private/AnimToTextureInstancePlaybackHelpers.cpp
+++ b/Engine/Plugins/Experimental/AnimToTexture/Source/AnimToTexture/Private/AnimToTextureInstancePlaybackHelpers.cpp
@@ -33,13 +33,13 @@ bool UAnimToTextureInstancePlaybackLibrary::SetupInstancedMeshComponent(UInstanc
 }
 
 bool UAnimToTextureInstancePlaybackLibrary::BatchUpdateInstancesAutoPlayData(UInstancedStaticMeshComponent* InstancedMeshComponent, 
-	const TArray<FAnimToTextureAutoPlayData>& AutoPlayData, const TArray<FMatrix>& Transforms, bool bMarkRenderStateDirty)
+	const TArray<FAnimToTextureAutoPlayData>& AutoPlayData, const TArray<FTransform>& Transforms, bool bMarkRenderStateDirty)
 {
 	return BatchUpdateInstancesData< FAnimToTextureAutoPlayData>(InstancedMeshComponent, AutoPlayData, Transforms, bMarkRenderStateDirty);
 }
 
 bool UAnimToTextureInstancePlaybackLibrary::BatchUpdateInstancesFrameData(UInstancedStaticMeshComponent* InstancedMeshComponent,
-	const TArray<FAnimToTextureFrameData>& FrameData, const TArray<FMatrix>& Transforms, bool bMarkRenderStateDirty)
+	const TArray<FAnimToTextureFrameData>& FrameData, const TArray<FTransform>& Transforms, bool bMarkRenderStateDirty)
 {
 	return BatchUpdateInstancesData<FAnimToTextureFrameData>(InstancedMeshComponent, FrameData, Transforms, bMarkRenderStateDirty);
 }
diff --git a/Engine/Plugins/Experimental/AnimToTexture/Source/AnimToTexture/Public/AnimToTextureInstancePlaybackHelpers.h b/Engine/Plugins/Experimental/AnimToTexture/Source/AnimToTexture/Public/AnimToTextureInstancePlaybackHelpers.h
--- a/Engine/Plugins/Experimental/AnimToTexture/Source/AnimToTexture/Public/AnimToTextureInstancePlaybackHelpers.h
+++ b/Engine/Plugins/Experimental/AnimToTexture/Source/AnimToTexture/Public/AnimToTextureInstancePlaybackHelpers.h
@@ -78,7 +78,7 @@ public:
 	*/
 	UFUNCTION(BlueprintCallable, Category = "AnimToTexture|Playback")
 	static bool BatchUpdateInstancesAutoPlayData(UInstancedStaticMeshComponent* InstancedMeshComponent,
-		const TArray<FAnimToTextureAutoPlayData>& AutoPlayData, const TArray<FMatrix>& Transforms, bool bMarkRenderStateDirty=true);
+		const TArray<FAnimToTextureAutoPlayData>& AutoPlayData, const TArray<FTransform>& Transforms, bool bMarkRenderStateDirty=true);
 
 	/**
 	* Updates all instances with the given Transforms and FrameData
@@ -86,7 +86,7 @@ public:
 	*/
 	UFUNCTION(BlueprintCallable, Category = "AnimToTexture|Playback")
 	static bool BatchUpdateInstancesFrameData(UInstancedStaticMeshComponent* InstancedMeshComponent,
-		const TArray<FAnimToTextureFrameData>& FrameData, const TArray<FMatrix>& Transforms, bool bMarkRenderStateDirty=true);
+		const TArray<FAnimToTextureFrameData>& FrameData, const TArray<FTransform>& Transforms, bool bMarkRenderStateDirty=true);
 
 	/**
 	* Updates a single instance with given AutoPlayData
@@ -118,13 +118,13 @@ public:
 
 	UFUNCTION(BlueprintCallable, Category = "AnimToTexture|Playback")
 	static bool GetFrameDataFromDataAsset(const UAnimToTextureDataAsset* DataAsset, int32 AnimationIndex, float Time,
-		FAnimToTextureFrameData& AutoPlayData, float TimeOffset = 0.f, float PlayRate = 1.f);
+		FAnimToTextureFrameData& FrameData, float TimeOffset = 0.f, float PlayRate = 1.f);
 	
 private:
 	
 	template <class DataType>
 	static bool BatchUpdateInstancesData(UInstancedStaticMeshComponent* InstancedMeshComponent,
-		const TArray<DataType>& Data, const TArray<FMatrix>& Transforms, bool bMarkRenderStateDirty = true);
+		const TArray<DataType>& Data, const TArray<FTransform>& Transforms, bool bMarkRenderStateDirty = true);
 	
 	template <class DataType>
 	static bool UpdateInstanceCustomData(UInstancedStaticMeshComponent* InstancedMeshComponent, int32 InstanceIndex,
@@ -141,7 +141,7 @@ private:
 
 template <class DataType>
 bool UAnimToTextureInstancePlaybackLibrary::BatchUpdateInstancesData(UInstancedStaticMeshComponent* InstancedMeshComponent,
-	const TArray<DataType>& Data, const TArray<FMatrix>& Transforms, bool bMarkRenderStateDirty)
+	const TArray<DataType>& Data, const TArray<FTransform>& Transforms, bool bMarkRenderStateDirty)
 {
 	if (!InstancedMeshComponent)
 	{
@@ -150,26 +150,24 @@ bool UAnimToTextureInstancePlaybackLibrary::BatchUpdateInstancesData(UInstancedS
 
 	// Get Number of Instances
 	const int32 NumInstances = InstancedMeshComponent->GetNumRenderInstances();
-
+	const int32 NumCustomDataFloats = InstancedMeshComponent->NumCustomDataFloats;
+	
 	// Check number of instances
 	if (NumInstances &&
 		NumInstances == Data.Num() &&
-		NumInstances == Transforms.Num())
+		NumInstances == Transforms.Num() &&
+		NumCustomDataFloats == Data.GetTypeSize() / sizeof(float))
 	{
-		// Check data size
-		if (InstancedMeshComponent->NumCustomDataFloats == Data.GetTypeSize() / sizeof(float))
-		{
-			SIZE_T CustomDataSizeToCopy = InstancedMeshComponent->PerInstanceSMCustomData.Num() * InstancedMeshComponent->PerInstanceSMCustomData.GetTypeSize();
-			FMemory::Memcpy(InstancedMeshComponent->PerInstanceSMCustomData.GetData(), Data.GetData(), CustomDataSizeToCopy);
+		// Set Transforms
+		// TODO: we should be passing PrevTransforms too.
+		InstancedMeshComponent->BatchUpdateInstancesTransforms(0, Transforms, false, bMarkRenderStateDirty, false);
 
-			// FInstancedStaticMeshInstanceData is not exposed to blueprints, so we are duplicating the data here :/
-			TArray<FInstancedStaticMeshInstanceData> InstanceData;
-			InstanceData.AddUninitialized(NumInstances);
-			FMemory::Memcpy(InstanceData.GetData(), Transforms.GetData(), NumInstances * InstanceData.GetTypeSize());
+		// Set Custom Data
+		TArray<float> CustomData;
+		GetCustomData(Data, CustomData);
+		InstancedMeshComponent->SetCustomData(0, NumInstances-1, CustomData, bMarkRenderStateDirty);
 
-			// Batch Update
-			return InstancedMeshComponent->BatchUpdateInstancesData(0, NumInstances, InstanceData.GetData(), bMarkRenderStateDirty, false);
-		}
+		return true;
 	}
 
 	return false;
