[AnimToTexture] Moving network from MaterialLayer into a MaterialFunction.
Keeping MaterialLayer for backward compatibility

#virtualized

diff --git a/Engine/Plugins/Experimental/AnimToTexture/Source/AnimToTextureEditor/Private/AssetDefinition_AnimToTexture.cpp b/Engine/Plugins/Experimental/AnimToTexture/Source/AnimToTextureEditor/Private/AssetDefinition_AnimToTexture.cpp
--- a/Engine/Plugins/Experimental/AnimToTexture/Source/AnimToTextureEditor/Private/AssetDefinition_AnimToTexture.cpp
+++ b/Engine/Plugins/Experimental/AnimToTexture/Source/AnimToTextureEditor/Private/AssetDefinition_AnimToTexture.cpp
@@ -59,7 +59,7 @@ namespace MenuExtension_AnimToTexture
 					{
 						if (UMaterialInstanceConstant* MaterialInstanceConstant = Cast<UMaterialInstanceConstant>(StaticMaterial.MaterialInterface))
 						{
-							UAnimToTextureBPLibrary::UpdateMaterialInstanceFromDataAsset(DataAsset, MaterialInstanceConstant, EMaterialParameterAssociation::LayerParameter);
+							UAnimToTextureBPLibrary::UpdateMaterialInstanceFromDataAsset(DataAsset, MaterialInstanceConstant, EMaterialParameterAssociation::GlobalParameter);
 						}
 					}
 				}
diff --git a/Engine/Plugins/Experimental/AnimToTexture/Source/AnimToTextureEditor/Public/AnimToTextureBPLibrary.h b/Engine/Plugins/Experimental/AnimToTexture/Source/AnimToTextureEditor/Public/AnimToTextureBPLibrary.h
--- a/Engine/Plugins/Experimental/AnimToTexture/Source/AnimToTextureEditor/Public/AnimToTextureBPLibrary.h
+++ b/Engine/Plugins/Experimental/AnimToTexture/Source/AnimToTextureEditor/Public/AnimToTextureBPLibrary.h
@@ -47,7 +47,7 @@ class UAnimToTextureBPLibrary : public UBlueprintFunctionLibrary
 	*/
 	UFUNCTION(BlueprintCallable, meta = (Category = "AnimToTexture"))
 	static ANIMTOTEXTUREEDITOR_API void UpdateMaterialInstanceFromDataAsset(const UAnimToTextureDataAsset* DataAsset, class UMaterialInstanceConstant* MaterialInstance,
-		const EMaterialParameterAssociation MaterialParameterAssociation = EMaterialParameterAssociation::LayerParameter);
+		const EMaterialParameterAssociation MaterialParameterAssociation = EMaterialParameterAssociation::GlobalParameter);
 
 #endif // WITH_EDITOR
 
