Ran UnrealCodeFixup to move dllstorage to correct places in code

diff --git a/Engine/Plugins/Animation/IKRig/Source/IKRig/Public/Rig/Solvers/IKRigStretchLimb.h b/Engine/Plugins/Animation/IKRig/Source/IKRig/Public/Rig/Solvers/IKRigStretchLimb.h
--- a/Engine/Plugins/Animation/IKRig/Source/IKRig/Public/Rig/Solvers/IKRigStretchLimb.h
+++ b/Engine/Plugins/Animation/IKRig/Source/IKRig/Public/Rig/Solvers/IKRigStretchLimb.h
@@ -1,4 +1,4 @@
-﻿// Copyright Epic Games, Inc. All Rights Reserved.
+// Copyright Epic Games, Inc. All Rights Reserved.
 
 #pragma once
 
@@ -99,11 +99,11 @@ struct FIKRigStretchLimbSolver : public FIKRigSolverBase
 	UE_API virtual void OnGoalMovedToDifferentBone(const FName& InGoalName, const FName& InNewBoneName) override;
 
 	// goal settings
-	UE_API virtual bool UsesCustomGoalSettings() const override { return false; };
+	virtual bool UsesCustomGoalSettings() const override { return false; };
 	
 	// start bone can be set on this solver
-	UE_API virtual bool UsesStartBone() const override { return true; };
-	UE_API virtual FName GetStartBone() const override { return Settings.StartBone; };
+	virtual bool UsesStartBone() const override { return true; };
+	virtual FName GetStartBone() const override { return Settings.StartBone; };
 	UE_API virtual void SetStartBone(const FName& InRootBoneName) override;
 
 #if WITH_EDITOR
diff --git a/Engine/Plugins/Experimental/AnimToTexture/Source/AnimToTexture/Public/AnimToTextureDataAsset.h b/Engine/Plugins/Experimental/AnimToTexture/Source/AnimToTexture/Public/AnimToTextureDataAsset.h
--- a/Engine/Plugins/Experimental/AnimToTexture/Source/AnimToTexture/Public/AnimToTextureDataAsset.h
+++ b/Engine/Plugins/Experimental/AnimToTexture/Source/AnimToTexture/Public/AnimToTextureDataAsset.h
@@ -9,6 +9,8 @@
 #include "Engine/Texture2D.h"
 #include "AnimToTextureDataAsset.generated.h"
 
+#define UE_API ANIMTOTEXTURE_API
+
 class USkeletalMesh;
 class UStaticMesh;
 class UTexture2D;
@@ -106,8 +108,8 @@ struct FAnimToTextureAnimInfo
 
 };
 
-UCLASS(Blueprintable, BlueprintType)
-class ANIMTOTEXTURE_API UAnimToTextureDataAsset : public UPrimaryDataAsset
+UCLASS(MinimalAPI, Blueprintable, BlueprintType)
+class UAnimToTextureDataAsset : public UPrimaryDataAsset
 {
 public:
 	GENERATED_BODY()
@@ -329,10 +331,10 @@ public:
 	*  Returns -1 if not found.
 	*/
 	UFUNCTION(BlueprintCallable, Category = Default)
-	int32 GetIndexFromAnimSequence(const UAnimSequence* Sequence);
+	UE_API int32 GetIndexFromAnimSequence(const UAnimSequence* Sequence);
 
 	UFUNCTION()
-	void ResetInfo();
+	UE_API void ResetInfo();
 
 	// If we weren't in a plugin, we could unify this in a base class
 	template<typename AssetType>
@@ -380,3 +382,5 @@ public:
 	UFUNCTION(BlueprintPure, Category = Default, meta = (DisplayName = "Get Bone Weight Texture"))
 	UTexture2D* BP_GetBoneWeightTexture() { return GetBoneWeightTexture(); }
 };
+
+#undef UE_API
diff --git a/Engine/Plugins/Experimental/AnimToTexture/Source/AnimToTexture/Public/AnimToTextureInstancePlaybackHelpers.h b/Engine/Plugins/Experimental/AnimToTexture/Source/AnimToTexture/Public/AnimToTextureInstancePlaybackHelpers.h
--- a/Engine/Plugins/Experimental/AnimToTexture/Source/AnimToTexture/Public/AnimToTextureInstancePlaybackHelpers.h
+++ b/Engine/Plugins/Experimental/AnimToTexture/Source/AnimToTexture/Public/AnimToTextureInstancePlaybackHelpers.h
@@ -7,6 +7,8 @@
 
 #include "AnimToTextureInstancePlaybackHelpers.generated.h"
 
+#define UE_API ANIMTOTEXTURE_API
+
 class UAnimToTextureDataAsset;
 
 USTRUCT(BlueprintType)
@@ -58,8 +60,8 @@ struct FAnimToTextureAutoPlayData
 };
 
 
-UCLASS()
-class ANIMTOTEXTURE_API UAnimToTextureInstancePlaybackLibrary : public UBlueprintFunctionLibrary
+UCLASS(MinimalAPI)
+class UAnimToTextureInstancePlaybackLibrary : public UBlueprintFunctionLibrary
 {
 	GENERATED_BODY()
 
@@ -70,14 +72,14 @@ public:
 	* @param bAutoPlay: if true, it will allocate the required CustomData for working with AutoPlayData. If false, FrameData will be allocated instead./
 	*/
 	UFUNCTION(BlueprintCallable, Category = "AnimToTexture|Playback")
-	static bool SetupInstancedMeshComponent(UInstancedStaticMeshComponent* InstancedMeshComponent, int32 NumInstances, bool bAutoPlay);
+	static UE_API bool SetupInstancedMeshComponent(UInstancedStaticMeshComponent* InstancedMeshComponent, int32 NumInstances, bool bAutoPlay);
 
 	/**
 	* Updates all instances with the given Transforms and AutoPlayData
 	* @param bMarkRenderStateDirty: if true, the change should be visible immediatelly. If you are updating many instances you should only set this to true for the last instance
 	*/
 	UFUNCTION(BlueprintCallable, Category = "AnimToTexture|Playback")
-	static bool BatchUpdateInstancesAutoPlayData(UInstancedStaticMeshComponent* InstancedMeshComponent,
+	static UE_API bool BatchUpdateInstancesAutoPlayData(UInstancedStaticMeshComponent* InstancedMeshComponent,
 		const TArray<FAnimToTextureAutoPlayData>& AutoPlayData, const TArray<FTransform>& Transforms, bool bMarkRenderStateDirty=true);
 
 	/**
@@ -85,7 +87,7 @@ public:
 	* @param bMarkRenderStateDirty: if true, the change should be visible immediatelly. If you are updating many instances you should only set this to true for the last instance
 	*/
 	UFUNCTION(BlueprintCallable, Category = "AnimToTexture|Playback")
-	static bool BatchUpdateInstancesFrameData(UInstancedStaticMeshComponent* InstancedMeshComponent,
+	static UE_API bool BatchUpdateInstancesFrameData(UInstancedStaticMeshComponent* InstancedMeshComponent,
 		const TArray<FAnimToTextureFrameData>& FrameData, const TArray<FTransform>& Transforms, bool bMarkRenderStateDirty=true);
 
 	/**
@@ -93,7 +95,7 @@ public:
 	* @param bMarkRenderStateDirty: if true, the change should be visible immediatelly. If you are updating many instances you should only set this to true for the last instance
 	*/
 	UFUNCTION(BlueprintCallable, Category = "AnimToTexture|Playback")
-	static bool UpdateInstanceAutoPlayData(UInstancedStaticMeshComponent* InstancedMeshComponent, int32 InstanceIndex, 
+	static UE_API bool UpdateInstanceAutoPlayData(UInstancedStaticMeshComponent* InstancedMeshComponent, int32 InstanceIndex, 
 		const FAnimToTextureAutoPlayData& AutoPlayData, bool bMarkRenderStateDirty=true);
 
 	/**
@@ -101,7 +103,7 @@ public:
 	* @param bMarkRenderStateDirty: if true, the change should be visible immediatelly. If you are updating many instances you should only set this to true for the last instance
 	*/
 	UFUNCTION(BlueprintCallable, Category = "AnimToTexture|Playback")
-	static bool UpdateInstanceFrameData(UInstancedStaticMeshComponent* InstancedMeshComponent, int32 InstanceIndex, 
+	static UE_API bool UpdateInstanceFrameData(UInstancedStaticMeshComponent* InstancedMeshComponent, int32 InstanceIndex, 
 		const FAnimToTextureFrameData& FrameData, bool bMarkRenderStateDirty=true);
 
 	/**
@@ -109,15 +111,15 @@ public:
 	* If AnimationIndex is out of range, false will be returned.
 	*/
 	UFUNCTION(BlueprintCallable, Category = "AnimToTexture|Playback")
-	static bool GetAutoPlayDataFromDataAsset(const UAnimToTextureDataAsset* DataAsset, int32 AnimationIndex,
+	static UE_API bool GetAutoPlayDataFromDataAsset(const UAnimToTextureDataAsset* DataAsset, int32 AnimationIndex,
 		FAnimToTextureAutoPlayData& AutoPlayData, float TimeOffset = 0.f, float PlayRate = 1.f);
 
 	UFUNCTION(BlueprintCallable, Category = "AnimToTexture|Playback")
-	static float GetFrame(float Time, float StartFrame, float EndFrame, 
+	static UE_API float GetFrame(float Time, float StartFrame, float EndFrame, 
 		float TimeOffset = 0.f, float PlayRate = 1.f, float SampleRate = 30.f);
 
 	UFUNCTION(BlueprintCallable, Category = "AnimToTexture|Playback")
-	static bool GetFrameDataFromDataAsset(const UAnimToTextureDataAsset* DataAsset, int32 AnimationIndex, float Time,
+	static UE_API bool GetFrameDataFromDataAsset(const UAnimToTextureDataAsset* DataAsset, int32 AnimationIndex, float Time,
 		FAnimToTextureFrameData& FrameData, float TimeOffset = 0.f, float PlayRate = 1.f);
 	
 private:
@@ -209,3 +211,5 @@ bool UAnimToTextureInstancePlaybackLibrary::UpdateInstanceCustomData(UInstancedS
 
 	return false;
 }
+
+#undef UE_API
