Fixup API macro usage

#rnx

diff --git a/Engine/Plugins/Experimental/AnimToTexture/Source/AnimToTexture/Private/AnimToTextureDataAsset.cpp b/Engine/Plugins/Experimental/AnimToTexture/Source/AnimToTexture/Private/AnimToTextureDataAsset.cpp
--- a/Engine/Plugins/Experimental/AnimToTexture/Source/AnimToTexture/Private/AnimToTextureDataAsset.cpp
+++ b/Engine/Plugins/Experimental/AnimToTexture/Source/AnimToTexture/Private/AnimToTextureDataAsset.cpp
@@ -2,6 +2,8 @@
 
 #include "AnimToTextureDataAsset.h"
 
+#include UE_INLINE_GENERATED_CPP_BY_NAME(AnimToTextureDataAsset)
+
 int32 UAnimToTextureDataAsset::GetIndexFromAnimSequence(const UAnimSequence* Sequence)
 {
 	const int32 NumSequences = AnimSequences.Num();
@@ -36,4 +38,4 @@ void UAnimToTextureDataAsset::ResetInfo()
 	BoneWeightRowsPerFrame = 1;
 	BoneMinBBox = FVector3f::ZeroVector;
 	BoneSizeBBox = FVector3f::ZeroVector;
-};
\ No newline at end of file
+};
diff --git a/Engine/Plugins/Experimental/AnimToTexture/Source/AnimToTexture/Private/AnimToTextureInstancePlaybackHelpers.cpp b/Engine/Plugins/Experimental/AnimToTexture/Source/AnimToTexture/Private/AnimToTextureInstancePlaybackHelpers.cpp
--- a/Engine/Plugins/Experimental/AnimToTexture/Source/AnimToTexture/Private/AnimToTextureInstancePlaybackHelpers.cpp
+++ b/Engine/Plugins/Experimental/AnimToTexture/Source/AnimToTexture/Private/AnimToTextureInstancePlaybackHelpers.cpp
@@ -3,6 +3,8 @@
 #include "AnimToTextureInstancePlaybackHelpers.h"
 #include "AnimToTextureDataAsset.h"
 
+#include UE_INLINE_GENERATED_CPP_BY_NAME(AnimToTextureInstancePlaybackHelpers)
+
 bool UAnimToTextureInstancePlaybackLibrary::SetupInstancedMeshComponent(UInstancedStaticMeshComponent* InstancedMeshComponent, int32 NumInstances, bool bAutoPlay)
 {	
 	if (!InstancedMeshComponent)
diff --git a/Engine/Source/Runtime/UMG/Public/Components/PostBufferUpdate.h b/Engine/Source/Runtime/UMG/Public/Components/PostBufferUpdate.h
--- a/Engine/Source/Runtime/UMG/Public/Components/PostBufferUpdate.h
+++ b/Engine/Source/Runtime/UMG/Public/Components/PostBufferUpdate.h
@@ -77,7 +77,7 @@ public:
 	UMG_API void SetPerformDefaultPostBufferUpdate(bool bInPerformDefaultPostBufferUpdate);
 
 #if WITH_EDITOR
-	UMG_API bool GetUpdateOnlyPaintArea() const { return bUpdateOnlyPaintArea; }
+	bool GetUpdateOnlyPaintArea() const { return bUpdateOnlyPaintArea; }
 
 	/** True if this widget updates any buffer indicated in SlateBufferBits */
 	UMG_API bool DoesUpdateSlatePostRT(ESlatePostRT SlateBufferBits) const ;
