[AnimToTexture] Fixed Baking RigidBodies

diff --git a/Engine/Plugins/Experimental/AnimToTexture/Source/AnimToTextureEditor/Private/AnimToTextureBPLibrary.cpp b/Engine/Plugins/Experimental/AnimToTexture/Source/AnimToTextureEditor/Private/AnimToTextureBPLibrary.cpp
--- a/Engine/Plugins/Experimental/AnimToTexture/Source/AnimToTextureEditor/Private/AnimToTextureBPLibrary.cpp
+++ b/Engine/Plugins/Experimental/AnimToTexture/Source/AnimToTextureEditor/Private/AnimToTextureBPLibrary.cpp
@@ -145,13 +145,14 @@ bool UAnimToTextureBPLibrary::AnimationToTexture(UAnimToTextureDataAsset* DataAs
 		{
 			AnimProgressBar.EnterProgressFrame();
 
-			const float Time = AnimStartTime + ((float)SampleIndex * SampleInterval);
+			const float DeltaTime = (float)SampleIndex * SampleInterval;
+			const float Time = AnimStartTime + DeltaTime;
 			SampleIndex++;
 
 			// Go To Time
 			SkeletalMeshComponent->SetPosition(Time);
 			// Update SkelMesh Animation.
-			SkeletalMeshComponent->TickAnimation(0.f, false /*bNeedsValidRootMotion*/);
+			SkeletalMeshComponent->TickAnimation(DeltaTime, false/*bNeedsValidRootMotion*/); // fix for evaluating rigid bodies in post process animbp
 			SkeletalMeshComponent->RefreshBoneTransforms(nullptr /*TickFunction*/);
 			
 			// ---------------------------------------------------------------------------
@@ -732,14 +733,14 @@ void UAnimToTextureBPLibrary::UpdateMaterialInstanceFromDataAsset(const UAnimToT
 	UMaterialEditingLibrary::RebuildMaterialInstanceEditors(MaterialInstance->GetMaterial());
 
 	// Set Preview Mesh
-	if (DataAsset->GetStaticMesh())
+	if(UStaticMesh* StaticMesh = DataAsset->GetStaticMesh())
 	{
-		MaterialInstance->PreviewMesh = DataAsset->GetStaticMesh();
+		MaterialInstance->PreviewMesh = FSoftObjectPath(StaticMesh);
 	}
 
 	MaterialInstance->MarkPackageDirty();
-}
 
+}
 
 bool UAnimToTextureBPLibrary::SetLightMapIndex(UStaticMesh* StaticMesh, const int32 LODIndex, const int32 LightmapIndex, bool bGenerateLightmapUVs)
 {
