[AnimToTexture] Added SkeletalMesh MeshDescription functions. (not used)

diff --git a/Engine/Plugins/Experimental/AnimToTexture/Source/AnimToTextureEditor/AnimToTextureEditor.Build.cs b/Engine/Plugins/Experimental/AnimToTexture/Source/AnimToTextureEditor/AnimToTextureEditor.Build.cs
--- a/Engine/Plugins/Experimental/AnimToTexture/Source/AnimToTextureEditor/AnimToTextureEditor.Build.cs
+++ b/Engine/Plugins/Experimental/AnimToTexture/Source/AnimToTextureEditor/AnimToTextureEditor.Build.cs
@@ -46,6 +46,7 @@ public class AnimToTextureEditor : ModuleRules
 				"AnimToTexture",
 				"RawMesh",
 				"MeshDescription",
+				//"SkeletalMeshDescription",
 				"StaticMeshDescription",
 				"ToolMenus",
 				"UnrealEd",
diff --git a/Engine/Plugins/Experimental/AnimToTexture/Source/AnimToTextureEditor/Private/AnimToTextureSkeletalMesh.cpp b/Engine/Plugins/Experimental/AnimToTexture/Source/AnimToTextureEditor/Private/AnimToTextureSkeletalMesh.cpp
--- a/Engine/Plugins/Experimental/AnimToTexture/Source/AnimToTextureEditor/Private/AnimToTextureSkeletalMesh.cpp
+++ b/Engine/Plugins/Experimental/AnimToTexture/Source/AnimToTextureEditor/Private/AnimToTextureSkeletalMesh.cpp
@@ -1,7 +1,6 @@
 // Copyright Epic Games, Inc. All Rights Reserved.
 #include "AnimToTextureSkeletalMesh.h"
 #include "AnimToTextureEditorModule.h"
-#include "Rendering/SkeletalMeshRenderData.h"
 #include "Engine/StaticMesh.h"
 #include "Engine/SkeletalMesh.h"
 #include "Components/SkeletalMeshComponent.h"
@@ -9,6 +8,8 @@
 #include "StaticMeshAttributes.h"
 #include "Algo/Reverse.h"
 #include "Math/NumericLimits.h"
+#include "Rendering/SkeletalMeshRenderData.h"
+//#include "SkeletalMeshAttributes.h"
 
 namespace AnimToTexture_Private
 {
@@ -28,13 +29,13 @@
 	// Get Mesh Description
 	const FMeshDescription* MeshDescription = StaticMesh->GetMeshDescription(LODIndex);
 	check(MeshDescription);
-	
+
 	// Get MeshDescription Vertices
 	const FVertexArray& Vertices = MeshDescription->Vertices();
 	const int32 NumVertices = Vertices.Num();
 	OutPositions.SetNum(NumVertices);
 	OutNormals.SetNumZeroed(NumVertices);
-	
+
 	for (const FVertexID& VertexID : Vertices.GetElementIDs())
 	{
 		const int32 VertexIndex = VertexID.GetValue();
@@ -50,16 +51,16 @@
 			OutNormals[VertexIndex] += VertexInstanceNormal / VertexInstances.Num();
 		}
 	}
-	
+
 	return NumVertices;
 }
 
-int32 GetVertices(const USkeletalMesh* SkeletalMesh, const int32 LODIndex, 
+int32 GetVertices(const USkeletalMesh* SkeletalMesh, const int32 LODIndex,
 	TArray<FVector3f>& OutPositions)
 {
 	check(SkeletalMesh);
 	OutPositions.Reset();
-	
+
 	const FSkeletalMeshRenderData* RenderData = SkeletalMesh->GetResourceForRendering();
 	check(RenderData);
 
@@ -74,7 +75,7 @@
 	// Get Total Num of Vertices (for all sections)
 	const int32 NumVertices = LODRenderData.GetNumVertices();
 	OutPositions.SetNumUninitialized(NumVertices);
-	
+
 	for (int32 VertexIndex = 0; VertexIndex < NumVertices; ++VertexIndex)
 	{
 		OutPositions[VertexIndex] = LODRenderData.StaticVertexBuffers.PositionVertexBuffer.VertexPosition(VertexIndex);
@@ -119,6 +120,175 @@
 	return NumTriangles;
 };
 
+void GetSkinWeights(const USkeletalMesh* SkeletalMesh, const int32 LODIndex,
+	TArray<VertexSkinWeightMax>& OutSkinWeights)
+{
+	check(SkeletalMesh);
+	OutSkinWeights.Reset();
+
+	// Get Render Data
+	const FSkeletalMeshRenderData* RenderData = SkeletalMesh->GetResourceForRendering();
+	check(RenderData);
+
+	if (!RenderData->LODRenderData.IsValidIndex(LODIndex))
+	{
+		UE_LOG(LogAnimToTextureEditor, Warning, TEXT("Invalid LODIndex: %i"), LODIndex)
+		return;
+	}
+
+	// Get LOD Data
+	const FSkeletalMeshLODRenderData& LODRenderData = RenderData->LODRenderData[LODIndex];
+
+	// Get Weights
+	const FSkinWeightVertexBuffer* SkinWeightVertexBuffer = LODRenderData.GetSkinWeightVertexBuffer();
+	check(SkinWeightVertexBuffer);
+
+	// Get Weights from Buffer.
+	// this is size of number of vertices.
+	TArray<FSkinWeightInfo> SkinWeightsInfo;
+	SkinWeightVertexBuffer->GetSkinWeights(SkinWeightsInfo);
+
+	// Allocated SkinWeightData
+	OutSkinWeights.SetNumUninitialized(SkinWeightsInfo.Num());
+
+	// Loop thru vertices
+	for (int32 VertexIndex = 0; VertexIndex < SkinWeightsInfo.Num(); VertexIndex++)
+	{
+		// Find Section From Global Vertex Index
+		// NOTE: BoneMap is stored by Section.
+		int32 OutSectionIndex;
+		int32 OutSectionVertexIndex;
+		LODRenderData.GetSectionFromVertexIndex(VertexIndex, OutSectionIndex, OutSectionVertexIndex);
+
+		// Get Section for Vertex.
+		const FSkelMeshRenderSection& RenderSection = LODRenderData.RenderSections[OutSectionIndex];
+
+		// Get Vertex Weights
+		const FSkinWeightInfo& SkinWeightInfo = SkinWeightsInfo[VertexIndex];
+
+		// Store Weights
+		for (int32 Index = 0; Index < MAX_TOTAL_INFLUENCES; Index++)
+		{
+			const uint8& BoneWeight = SkinWeightInfo.InfluenceWeights[Index];
+			const uint16& BoneIndex = SkinWeightInfo.InfluenceBones[Index];
+			const uint16& MeshBoneIndex = RenderSection.BoneMap[BoneIndex];
+
+			OutSkinWeights[VertexIndex].BoneWeights[Index] = BoneWeight;
+			OutSkinWeights[VertexIndex].MeshBoneIndices[Index] = MeshBoneIndex;
+		}
+	}
+}
+
+// NOTE: Unfortunatelly SkeletalMeshDescription is only generated for LOD0. 
+//       Leaving this here for future reference.
+/*
+int32 GetVertices(const USkeletalMesh* SkeletalMesh, const int32 LODIndex,
+	TArray<FVector3f>& OutPositions)
+{
+	check(SkeletalMesh);
+	OutPositions.Reset();
+
+	if (!SkeletalMesh->IsValidLODIndex(LODIndex))
+	{
+		return INDEX_NONE;
+	}
+
+	// Get Mesh Description
+	const FMeshDescription* MeshDescription = SkeletalMesh->GetMeshDescription(LODIndex);
+	check(MeshDescription)
+
+	const FVertexArray& Vertices = MeshDescription->Vertices();
+	const int32 NumVertices = Vertices.Num();
+
+	OutPositions.SetNum(NumVertices);
+
+	for (const FVertexID& VertexID : Vertices.GetElementIDs())
+	{
+		const int32 VertexIndex = VertexID.GetValue();
+		OutPositions[VertexIndex] = MeshDescription->GetVertexPosition(VertexID);
+	}
+
+	return NumVertices;
+}
+
+int32 GetTriangles(const USkeletalMesh* SkeletalMesh, const int32 LODIndex,
+	TArray<FIntVector3>& OutTriangles)
+{
+	check(SkeletalMesh);
+	OutTriangles.Reset();
+
+	if (!SkeletalMesh->IsValidLODIndex(LODIndex))
+	{
+		return INDEX_NONE;
+	}
+
+	const FMeshDescription* MeshDescription = SkeletalMesh->GetMeshDescription(LODIndex);
+	check(MeshDescription)
+
+	const FTriangleArray& Triangles = MeshDescription->Triangles();
+	const int32 NumTriangles = Triangles.Num();
+	OutTriangles.SetNumUninitialized(NumTriangles);
+
+	for (const FTriangleID& TriangleID : Triangles.GetElementIDs())
+	{
+		const TArrayView<const FVertexID>& TriangleVertices = MeshDescription->GetTriangleVertices(TriangleID);
+
+		const int32 TriangleIndex = TriangleID.GetValue();
+		OutTriangles[TriangleIndex].X = TriangleVertices[0].GetValue();
+		OutTriangles[TriangleIndex].Y = TriangleVertices[1].GetValue();
+		OutTriangles[TriangleIndex].Z = TriangleVertices[2].GetValue();
+	}
+
+	return NumTriangles;
+};
+
+
+void GetSkinWeights(const USkeletalMesh* SkeletalMesh, const int32 LODIndex,
+	TArray<VertexSkinWeightMax>& OutSkinWeights)
+{
+	check(SkeletalMesh);
+	OutSkinWeights.Reset();
+
+	if (!SkeletalMesh->IsValidLODIndex(LODIndex))
+	{
+		return;
+	}
+
+	FMeshDescription* MeshDescription = SkeletalMesh->GetMeshDescription(LODIndex);
+
+	// Get Vertices
+	const FVertexArray& Vertices = MeshDescription->Vertices();
+	const int32 NumVertices = Vertices.Num();
+
+	// Get Vertex Weights
+	FSkeletalMeshAttributes SkeletalMeshAttributes(*MeshDescription);
+	const FSkinWeightsVertexAttributesRef SkinWeights = SkeletalMeshAttributes.GetVertexSkinWeights();
+	const FSkeletalMeshAttributes::FBoneNameAttributesRef BoneNames = SkeletalMeshAttributes.GetBoneNames();
+
+	const int32 NumBones = SkeletalMeshAttributes.GetNumBones();
+
+	// Allocated SkinWeightData
+	OutSkinWeights.SetNumZeroed(NumVertices);
+
+	for (const FVertexID& VertexID : Vertices.GetElementIDs())
+	{
+		const int32 VertexIndex = VertexID.GetValue();
+		const FVertexBoneWeights BoneWeights = SkinWeights.Get(VertexID);
+		check(BoneWeights.Num() < MAX_TOTAL_INFLUENCES)
+
+		for (int32 WeightIndex = 0; WeightIndex < BoneWeights.Num(); WeightIndex++)
+		{
+			// Note: BoneIndex == RawBoneIndex
+			const int32 BoneIndex = BoneWeights[WeightIndex].GetBoneIndex();
+			const float BoneWeight = BoneWeights[WeightIndex].GetWeight();
+
+			OutSkinWeights[VertexIndex].BoneWeights[WeightIndex] = (uint8)FMath::RoundToInt(BoneWeight * 255.f);
+			OutSkinWeights[VertexIndex].MeshBoneIndices[WeightIndex] = BoneIndex;
+		}
+	}
+}
+*/
+
 int32 GetNumBones(const USkeletalMesh* SkeletalMesh)
 {
 	check(SkeletalMesh);
@@ -165,23 +335,6 @@
 	}
 }
 
-bool HasBone(const USkeletalMesh* SkeletalMesh, const FName& Bone)
-{
-	check(SkeletalMesh);
-	
-	const FReferenceSkeleton& RefSkeleton = SkeletalMesh->GetRefSkeleton();
-
-	for (const FMeshBoneInfo& BoneInfo: RefSkeleton.GetRawRefBoneInfo())
-	{
-		if (BoneInfo.Name == Bone)
-		{
-			return true;
-		}
-	}
-	return false;
-}
-
-
 void GetBoneNames(const USkeletalMesh* SkeletalMesh, TArray<FName>& OutNames)
 {
 	check(SkeletalMesh);
@@ -197,69 +350,12 @@
 		OutNames[BoneIndex] = RefSkeleton.GetBoneName(BoneIndex);
 	}
 }
+
 FVector3f PointAtBarycentricCoordinates(const FVector3f& PointA, const FVector3f& PointB, const FVector3f& PointC, const FVector3f& BarycentricCoords)
 {
 	return (PointA * BarycentricCoords.X) + (PointB * BarycentricCoords.Y) + (PointC * (1.f - BarycentricCoords.X - BarycentricCoords.Y));
 };
 
-void GetSkinWeights(const USkeletalMesh* SkeletalMesh, const int32 LODIndex, 
-	TArray<VertexSkinWeightMax>& OutSkinWeights)
-{
-	check(SkeletalMesh);
-	OutSkinWeights.Reset();
 
-	// Get Render Data
-	const FSkeletalMeshRenderData* RenderData = SkeletalMesh->GetResourceForRendering();
-	check(RenderData);
-
-	if (!RenderData->LODRenderData.IsValidIndex(LODIndex))
-	{
-		UE_LOG(LogAnimToTextureEditor, Warning, TEXT("Invalid LODIndex: %i"), LODIndex)
-		return;
-	}
-
-	// Get LOD Data
-	const FSkeletalMeshLODRenderData& LODRenderData = RenderData->LODRenderData[LODIndex];
-
-	// Get Weights
-	const FSkinWeightVertexBuffer* SkinWeightVertexBuffer = LODRenderData.GetSkinWeightVertexBuffer();
-	check(SkinWeightVertexBuffer);
-
-	// Get Weights from Buffer.
-	// this is size of number of vertices.
-	TArray<FSkinWeightInfo> SkinWeightsInfo;
-	SkinWeightVertexBuffer->GetSkinWeights(SkinWeightsInfo);
-
-	// Allocated SkinWeightData
-	OutSkinWeights.SetNumUninitialized(SkinWeightsInfo.Num());
-
-	// Loop thru vertices
-	for (int32 VertexIndex = 0; VertexIndex < SkinWeightsInfo.Num(); VertexIndex++)
-	{
-		// Find Section From Global Vertex Index
-		// NOTE: BoneMap is stored by Section.
-		int32 OutSectionIndex;
-		int32 OutSectionVertexIndex;
-		LODRenderData.GetSectionFromVertexIndex(VertexIndex, OutSectionIndex, OutSectionVertexIndex);
-
-		// Get Section for Vertex.
-		const FSkelMeshRenderSection& RenderSection = LODRenderData.RenderSections[OutSectionIndex];
-		
-		// Get Vertex Weights
-		const FSkinWeightInfo& SkinWeightInfo = SkinWeightsInfo[VertexIndex];
-
-		// Store Weights
-		for (int32 Index = 0; Index < MAX_TOTAL_INFLUENCES; Index++)
-		{
-			const uint8& BoneWeight = SkinWeightInfo.InfluenceWeights[Index];
-			const uint16& BoneIndex = SkinWeightInfo.InfluenceBones[Index];
-			const uint16& MeshBoneIndex = RenderSection.BoneMap[BoneIndex];
-
-			OutSkinWeights[VertexIndex].BoneWeights[Index] = BoneWeight;
-			OutSkinWeights[VertexIndex].MeshBoneIndices[Index] = MeshBoneIndex;
-		}
-	}
-}
-
 void InterpolateVertexSkinWeights(const TArray<VertexSkinWeightMax>& SkinWeights, const TArray<float>& Weights,
 	VertexSkinWeightMax& OutVertexSkinWeights)
@@ -383,11 +479,15 @@
 			const uint8& BoneWeight = Weights.BoneWeights[Index];
 			const uint16& MeshBoneIndex = Weights.MeshBoneIndices[Index];
 
-			// Get Matrix
-			const FMatrix44f& RefToLocal = RefToLocals[MeshBoneIndex];
+			// no need to transform if bone doesnt have weights.
+			if (BoneWeight > 0) 
+			{	
+				// get matrix
+				const FMatrix44f& RefToLocal = RefToLocals[MeshBoneIndex];
 
-			const float Weight = (float)BoneWeight / 255.f;
-			SkinnedVertex += RefToLocal.TransformPosition(Vertex) * Weight;
+				const float Weight = (float)BoneWeight / 255.f;
+				SkinnedVertex += RefToLocal.TransformPosition(Vertex) * Weight;
+			}
 		}
 
 		OutPositions[VertexIndex] = SkinnedVertex;
@@ -558,5 +658,4 @@
 	}
 }
 
-
 } // end namespace AnimToTexture_Private
diff --git a/Engine/Plugins/Experimental/AnimToTexture/Source/AnimToTextureEditor/Public/AnimToTextureSkeletalMesh.h b/Engine/Plugins/Experimental/AnimToTexture/Source/AnimToTextureEditor/Public/AnimToTextureSkeletalMesh.h
--- a/Engine/Plugins/Experimental/AnimToTexture/Source/AnimToTextureEditor/Public/AnimToTextureSkeletalMesh.h
+++ b/Engine/Plugins/Experimental/AnimToTexture/Source/AnimToTextureEditor/Public/AnimToTextureSkeletalMesh.h
@@ -55,9 +55,6 @@ int32 GetNumBones(const USkeletalMesh* SkeletalMesh);
    The returned transforms are in ComponentSpace */
 void GetRefBoneTransforms(const USkeletalMesh* SkeletalMesh, TArray<FTransform>& OutTransforms);
 
-/* Returns Bone exist in the RawBone list (no virtual bones) */
-bool HasBone(const USkeletalMesh* SkeletalMesh, const FName& Bone);
-
 /* Returns Bone Names. 
    Only the RawBones are returned (no virtual bones) */
 void GetBoneNames(const USkeletalMesh* SkeletalMesh, TArray<FName>& OutNames);
